Movie Render Queue:  Support for auto-exposure with high-res tiling and panoramic rendering.  An extra full screen render pass is run, and exposure is shared with all tiles from that.  For panoramic rendering, the full screen pass is 6 cube faces, similar to a cube scene capture.  Scene Capture Cube and MRQ panoramic rendering were both upgraded to generate auto-exposure from all 6 cube faces, eliminating any directional dependency of the auto-exposure calculation.

diff --git a/Engine/Plugins/MovieScene/MoviePipelineMaskRenderPass/Source/MoviePipelineMaskRenderPass/Private/MoviePipelinePanoramicPass.cpp b/Engine/Plugins/MovieScene/MoviePipelineMaskRenderPass/Source/MoviePipelineMaskRenderPass/Private/MoviePipelinePanoramicPass.cpp
--- a/Engine/Plugins/MovieScene/MoviePipelineMaskRenderPass/Source/MoviePipelineMaskRenderPass/Private/MoviePipelinePanoramicPass.cpp
+++ b/Engine/Plugins/MovieScene/MoviePipelineMaskRenderPass/Source/MoviePipelineMaskRenderPass/Private/MoviePipelinePanoramicPass.cpp
@@ -28,6 +28,7 @@ UMoviePipelinePanoramicPass::UMoviePipelinePanoramicPass()
 	, EyeSeparation(6.5f)
 	, EyeConvergenceDistance(EyeSeparation * 30.f)
 	, bAllocateHistoryPerPane(false)
+	, bAutoExposureRenderPass(true)
 	, bHasWarnedSettings(false)
 {
 	PassIdentifier = FMoviePipelinePassIdentifier("Panoramic");
@@ -90,6 +91,12 @@ namespace MoviePipeline
 	}
 }
 
+// Cube capture is arranged in 3x2 square tiles, rounded down to a multiple of 8 pixels.
+int32 ComputeAutoExposureCubeCaptureSize(FIntPoint Resolution)
+{
+	return AlignDown(FMath::Min(Resolution.X / 3, Resolution.Y / 2), 8);
+}
+
 void UMoviePipelinePanoramicPass::MoviePipelineRenderShowFlagOverride(FEngineShowFlags& OutShowFlag)
 {
 	// Panoramics can't support any of these.
@@ -127,6 +134,23 @@ void UMoviePipelinePanoramicPass::SetupImpl(const MoviePipeline::FMoviePipelineR
 		{
 			OptionalPaneViewStates[Index].Allocate(InPassInitSettings.FeatureLevel);
 		}
+
+		if (bAutoExposureRenderPass && ComputeAutoExposureCubeCaptureSize(PaneResolution) > 0)
+		{
+			AutoExposureViewStates.SetNum(CubeFace_MAX);
+
+			// ShareOrigin must be called before Allocate.  ShareOrigin is necessary for Lumen to work with 6 way cube split screen
+			// (causes Lumen scene data to be shared for all views, and overrides Lumen's regular 2 view split screen limitation).
+			for (int32 Index = 1; Index < CubeFace_MAX; Index++)
+			{
+				AutoExposureViewStates[Index].ShareOrigin(&AutoExposureViewStates[0]);
+			}
+
+			for (int32 Index = 0; Index < CubeFace_MAX; Index++)
+			{
+				AutoExposureViewStates[Index].Allocate(InPassInitSettings.FeatureLevel);
+			}
+		}
 	}
 
 	// We need one accumulator per pano tile if using accumulation.
@@ -157,7 +181,24 @@ void UMoviePipelinePanoramicPass::TeardownImpl()
 		OptionalPaneViewStates[Index].Destroy();
 	}
 	OptionalPaneViewStates.Reset();
-	
+
+	if (AutoExposureViewStates.Num() > 0)
+	{
+		check(AutoExposureViewStates.Num() == CubeFace_MAX)
+		for (int32 Index = 0; Index < CubeFace_MAX; Index++)
+		{
+			FSceneViewStateInterface* Ref = AutoExposureViewStates[Index].GetReference();
+			if (Ref)
+			{
+				Ref->ClearMIDPool();
+			}
+		}
+
+		// View states using FSceneViewStateReference::ShareOrigin need to be destroyed before their target, so remove last 5 elements first
+		AutoExposureViewStates.RemoveAt(1, CubeFace_MAX - 1);
+		AutoExposureViewStates.Reset();
+	}
+
 	OCIOSceneViewExtension.Reset();
 	OCIOSceneViewExtension = nullptr;
 
@@ -177,6 +218,15 @@ void UMoviePipelinePanoramicPass::AddReferencedObjects(UObject* InThis, FReferen
 			Ref->AddReferencedObjects(Collector);
 		}
 	}
+
+	for (int32 Index = 0; Index < This.AutoExposureViewStates.Num(); Index++)
+	{
+		FSceneViewStateInterface* Ref = This.AutoExposureViewStates[Index].GetReference();
+		if (Ref)
+		{
+			Ref->AddReferencedObjects(Collector);
+		}
+	}
 }
 
 FSceneViewStateInterface* UMoviePipelinePanoramicPass::GetSceneViewStateInterface(IViewCalcPayload* OptPayload)
@@ -194,6 +244,19 @@ FSceneViewStateInterface* UMoviePipelinePanoramicPass::GetSceneViewStateInterfac
 	}
 }
 
+FSceneViewStateInterface* UMoviePipelinePanoramicPass::GetExposureSceneViewStateInterface(IViewCalcPayload* OptPayload, int32 CubeFaceIndex)
+{
+	if (AutoExposureViewStates.Num() > 0)
+	{
+		// Index -1 is passed in when not a cube face, so use zero for that, as the eye adaptation source
+		return AutoExposureViewStates[FMath::Max(CubeFaceIndex,0)].GetReference();
+	}
+	else
+	{
+		return nullptr;
+	}
+}
+
 void UMoviePipelinePanoramicPass::GatherOutputPassesImpl(TArray<FMoviePipelinePassIdentifier>& ExpectedRenderPasses)
 {
 	// Despite rendering many views, we only output one image total, which is covered by the base.
@@ -253,25 +316,63 @@ FSceneView* UMoviePipelinePanoramicPass::GetSceneViewForSampleState(FSceneViewFa
 	//We skip calling the Super::GetSceneViewForSampleState entirely, as this has highly customized logic for Panos.
 	APlayerController* LocalPlayerController = GetPipeline()->GetWorld()->GetFirstPlayerController();
 
-	// We ignore the resolution that comes from the main render job as that will be our output resolution.
-	int32 PaneSizeX = PanoPane->Resolution.X;
-	int32 PaneSizeY = PanoPane->Resolution.Y;
+	int32 PaneSizeX;
+	int32 PaneSizeY;
+	FIntPoint ViewRectOffset;
+
+	if (InOutSampleState.bAutoExposureCubePass)
+	{
+		int32 CubeCaptureSize = ComputeAutoExposureCubeCaptureSize(PanoPane->Resolution);
+		check(CubeCaptureSize > 0);
+
+		PaneSizeX = CubeCaptureSize;
+		PaneSizeY = CubeCaptureSize;
+
+		// Auto exposure cube map faces are rendered as 3x2 split screen tiles.
+		static const FIntPoint GCubeFaceViewRectOffsets[6] =
+		{
+			{ 0,0 },
+			{ 1,0 },
+			{ 2,0 },
+			{ 0,1 },
+			{ 1,1 },
+			{ 2,1 },
+		};
+		ViewRectOffset = GCubeFaceViewRectOffsets[InOutSampleState.AutoExposureCubeFace] * FIntPoint(PaneSizeX, PaneSizeY);
+	}
+	else
+	{
+		// We ignore the resolution that comes from the main render job as that will be our output resolution.
+		PaneSizeX = PanoPane->Resolution.X;
+		PaneSizeY = PanoPane->Resolution.Y;
+
+		ViewRectOffset = FIntPoint(0, 0);
+	}
 
 	FSceneViewInitOptions ViewInitOptions;
 	ViewInitOptions.ViewFamily = ViewFamily;
 	ViewInitOptions.ViewOrigin = PanoPane->CameraLocation;
-	ViewInitOptions.SetViewRectangle(FIntRect(FIntPoint(0, 0), FIntPoint(PaneSizeX, PaneSizeY)));
-	ViewInitOptions.ViewRotationMatrix = FInverseRotationMatrix(FRotator(PanoPane->CameraRotation));
+	ViewInitOptions.SetViewRectangle(FIntRect(ViewRectOffset, ViewRectOffset + FIntPoint(PaneSizeX, PaneSizeY)));
 	ViewInitOptions.ViewActor = LocalPlayerController ? LocalPlayerController->GetViewTarget() : nullptr;
 
-	// Rotate the view 90 degrees (reason: unknown)
-	ViewInitOptions.ViewRotationMatrix = ViewInitOptions.ViewRotationMatrix * FMatrix(
-		FPlane(0, 0, 1, 0),
-		FPlane(1, 0, 0, 0),
-		FPlane(0, 1, 0, 0),
-		FPlane(0, 0, 0, 1));
+	if (InOutSampleState.bAutoExposureCubePass)
+	{
+		ViewInitOptions.ViewRotationMatrix = CalcCubeFaceTransform((ECubeFace)InOutSampleState.AutoExposureCubeFace);
+	}
+	else
+	{
+		ViewInitOptions.ViewRotationMatrix = FInverseRotationMatrix(FRotator(PanoPane->CameraRotation));
+
+		// Rotate the view 90 degrees (reason: unknown)
+		ViewInitOptions.ViewRotationMatrix = ViewInitOptions.ViewRotationMatrix * FMatrix(
+			FPlane(0, 0, 1, 0),
+			FPlane(1, 0, 0, 0),
+			FPlane(0, 1, 0, 0),
+			FPlane(0, 0, 0, 1));
+	}
+
+	float ViewFOV = InOutSampleState.bAutoExposureCubePass ? 90.f : PanoPane->HorizontalFieldOfView;
 
-	float ViewFOV = PanoPane->HorizontalFieldOfView;
 	// Inflate our FOV to support the overscan  
 	// Overscan not supported right now.
 	// ViewFOV = 2.0f * FMath::RadiansToDegrees(FMath::Atan((1.0f + InOutSampleState.OverscanPercentage) * FMath::Tan(FMath::DegreesToRadians(ViewFOV * 0.5f))));
@@ -335,15 +436,28 @@ FSceneView* UMoviePipelinePanoramicPass::GetSceneViewForSampleState(FSceneViewFa
 		ViewInitOptions.ProjectionMatrix = BaseProjMatrix;
 	}
 
-	ViewInitOptions.SceneViewStateInterface = GetSceneViewStateInterface(OptPayload);
+	ViewInitOptions.SceneViewStateInterface = InOutSampleState.bAutoExposureCubePass ? GetExposureSceneViewStateInterface(OptPayload, InOutSampleState.AutoExposureCubeFace) : GetSceneViewStateInterface(OptPayload);
 	ViewInitOptions.FOV = ViewFOV;
 
+	// If exposure isn't active, this will return nullptr
+	ViewInitOptions.ExposureSceneViewStateInterface = GetExposureSceneViewStateInterface(OptPayload, 0);
+
 	FSceneView* View = new FSceneView(ViewInitOptions);
 	ViewFamily->Views.Add(View);
 	View->ViewLocation = ViewInitOptions.ViewOrigin;
-	View->ViewRotation = PanoPane->CameraRotation;
-	// Override previous/current view transforms so that tiled renders don't use the wrong occlusion/motion blur information.
-	View->PreviousViewTransform = FTransform(PanoPane->PrevCameraRotation, PanoPane->PrevCameraLocation);
+
+	if (InOutSampleState.bAutoExposureCubePass)
+	{
+		// Rotation is world space for exposure cube map passes, generate it from the matrix.  Also, it's invariant, so previous rotation is the same as current.
+		View->ViewRotation = ViewInitOptions.ViewRotationMatrix.Rotator();
+		View->PreviousViewTransform = FTransform(View->ViewRotation, PanoPane->PrevCameraLocation);
+	}
+	else
+	{
+		View->ViewRotation = PanoPane->CameraRotation;
+		// Override previous/current view transforms so that tiled renders don't use the wrong occlusion/motion blur information.
+		View->PreviousViewTransform = FTransform(PanoPane->PrevCameraRotation, PanoPane->PrevCameraLocation);
+	}
 
 	View->StartFinalPostprocessSettings(View->ViewLocation);
 	BlendPostProcessSettings(View, InOutSampleState, OptPayload);
@@ -352,6 +466,19 @@ FSceneView* UMoviePipelinePanoramicPass::GetSceneViewForSampleState(FSceneViewFa
 	// cause the circle of confusion to be unchanged.
 	View->FinalPostProcessSettings.DepthOfFieldSensorWidth *= DofSensorScale;
 
+	// Disable anti-aliasing and temporal upscale for auto-exposure passes.  Auto-exposure is calculated before those passes, so this is wasted work (and memory for history).
+	if (InOutSampleState.bAutoExposureCubePass)
+	{
+		View->AntiAliasingMethod = AAM_None;
+		View->PrimaryScreenPercentageMethod = EPrimaryScreenPercentageMethod::SpatialUpscale;
+
+		// Generate auto-exposure from all cube map faces for the first view
+		if (InOutSampleState.AutoExposureCubeFace == 0)
+		{
+			View->bEyeAdaptationAllViewPixels = true;
+		}
+	}
+
 	// Modify the 'center' of the lens to be offset for high-res tiling, helps some effects (vignette) etc. still work.
 	// ToDo: Highres Tiling support
 	// View->LensPrincipalPointOffsetScale = CalculatePrinciplePointOffsetForTiling(InOutSampleState);
@@ -391,6 +518,51 @@ void UMoviePipelinePanoramicPass::RenderSample_GameThreadImpl(const FMoviePipeli
 
 	const FIntPoint PaneResolution = GetPaneResolution(InSampleState.BackbufferSize);
 	
+	// Auto-exposure is rendered for the first sample, and shared across all panes
+	if (AutoExposureViewStates.Num() && InSampleState.SpatialSampleIndex == 0 && InSampleState.TemporalSampleIndex == 0)
+	{
+		FMoviePipelineRenderPassMetrics InOutSampleState = InSampleState;
+		InOutSampleState.bAutoExposureCubePass = true;
+		InOutSampleState.BackbufferSize = PaneResolution;
+		InOutSampleState.TileIndexes = FIntPoint(-1, -1);		// Marks this as the auto exposure pass
+
+		// Construct a pane that will be attached to the render information.  Most of the fields here aren't used for the auto exposure pass...
+		FPanoPane Pane;
+		Pane.OriginalCameraLocation = InSampleState.FrameInfo.CurrViewLocation;
+		Pane.PrevOriginalCameraLocation = InSampleState.FrameInfo.PrevViewLocation;
+		Pane.OriginalCameraRotation = InSampleState.FrameInfo.CurrViewRotation;
+		Pane.PrevOriginalCameraRotation = InSampleState.FrameInfo.PrevViewRotation;
+		Pane.EyeIndex = -1;
+		Pane.VerticalStepIndex = 0;
+		Pane.HorizontalStepIndex = 0;
+		Pane.NumHorizontalSteps = 1;
+		Pane.NumVerticalSteps = 1;
+		Pane.EyeSeparation = EyeSeparation;
+		Pane.EyeConvergenceDistance = EyeConvergenceDistance;
+		Pane.bUseLocalRotation = bFollowCameraOrientation;
+		
+		MoviePipeline::Panoramic::GetCameraOrientationForStereo(/*Out*/ Pane.CameraLocation, /*Out*/ Pane.CameraRotation, /*Out*/ Pane.CameraLocalRotation, Pane, -1, /*bInPrevPos*/ false);
+		FRotator DummyPrevLocalRot;
+		MoviePipeline::Panoramic::GetCameraOrientationForStereo(/*Out*/ Pane.PrevCameraLocation, /*Out*/ Pane.PrevCameraRotation, /*Out*/ DummyPrevLocalRot, Pane, -1, /*bInPrevPos*/ true);
+
+		Pane.HorizontalFieldOfView = 90.0f;
+		Pane.VerticalFieldOfView = 90.0f;
+		
+		Pane.Resolution = PaneResolution;
+
+		TSharedPtr<FSceneViewFamilyContext> ViewFamily = CalculateViewFamily(InOutSampleState, &Pane);
+
+		// Submit the view to be rendered.
+		TWeakObjectPtr<UTextureRenderTarget2D> ViewRenderTarget = GetOrCreateViewRenderTarget(PaneResolution);
+		check(ViewRenderTarget.IsValid());
+
+		FRenderTarget* RenderTarget = ViewRenderTarget->GameThread_GetRenderTargetResource();
+		check(RenderTarget);
+
+		FCanvas Canvas = FCanvas(RenderTarget, nullptr, GetPipeline()->GetWorld(), ViewFamily->GetFeatureLevel(), FCanvas::CDM_DeferDrawing, 1.0f);
+		GetRendererModule().BeginRenderingViewFamily(&Canvas, ViewFamily.Get());
+	}
+
 	// This function will be called for each sample of every tile. We'll then submit renders for each
 	// Panoramic Pane. This effectively means that all panes will fill in their top left corners first.
 	// The accumulation stage will sort out waiting for all the high-res tiles to come in, and once those
@@ -407,6 +579,7 @@ void UMoviePipelinePanoramicPass::RenderSample_GameThreadImpl(const FMoviePipeli
 			for (int32 EyeLoopIndex = 0; EyeLoopIndex < NumEyeRenders; EyeLoopIndex++)
 			{
 				FMoviePipelineRenderPassMetrics InOutSampleState = InSampleState;
+				InOutSampleState.bAutoExposureRenderPass = AutoExposureViewStates.Num() > 0;
 
 				// Construct a pane that will be attached to the render information
 				FPanoPane Pane;
diff --git a/Engine/Plugins/MovieScene/MoviePipelineMaskRenderPass/Source/MoviePipelineMaskRenderPass/Private/MoviePipelinePanoramicPass.h b/Engine/Plugins/MovieScene/MoviePipelineMaskRenderPass/Source/MoviePipelineMaskRenderPass/Private/MoviePipelinePanoramicPass.h
--- a/Engine/Plugins/MovieScene/MoviePipelineMaskRenderPass/Source/MoviePipelineMaskRenderPass/Private/MoviePipelinePanoramicPass.h
+++ b/Engine/Plugins/MovieScene/MoviePipelineMaskRenderPass/Source/MoviePipelineMaskRenderPass/Private/MoviePipelinePanoramicPass.h
@@ -127,8 +127,9 @@ protected:
 	virtual int32 GetOutputFileSortingOrder() const override { return 1; }
 	virtual bool IsAlphaInTonemapperRequiredImpl() const override { return false; }
 	virtual FSceneViewStateInterface* GetSceneViewStateInterface(IViewCalcPayload* OptPayload) override;
+	virtual FSceneViewStateInterface* GetExposureSceneViewStateInterface(IViewCalcPayload* OptPayload, int32 CubeFaceIndex) override;
 	virtual void AddViewExtensions(FSceneViewFamilyContext& InContext, FMoviePipelineRenderPassMetrics& InOutSampleState) override;
-	virtual bool IsAutoExposureAllowed(const FMoviePipelineRenderPassMetrics& InSampleState) const override { return false; }
+	virtual bool IsAutoExposureAllowed(const FMoviePipelineRenderPassMetrics& InSampleState) const override { return AutoExposureViewStates.Num() > 0; }
 	virtual FSceneView* GetSceneViewForSampleState(FSceneViewFamily* ViewFamily, FMoviePipelineRenderPassMetrics& InOutSampleState, IViewCalcPayload* OptPayload = nullptr) override;
 	virtual TWeakObjectPtr<UTextureRenderTarget2D> GetOrCreateViewRenderTarget(const FIntPoint& InSize, IViewCalcPayload* OptPayload = nullptr) override;
 	virtual TSharedPtr<FMoviePipelineSurfaceQueue, ESPMode::ThreadSafe> GetOrCreateSurfaceQueue(const FIntPoint& InSize, IViewCalcPayload* OptPayload = nullptr) override;
@@ -197,12 +198,22 @@ public:
 	UPROPERTY(EditAnywhere, BlueprintReadWrite, AdvancedDisplay, Category = "Panoramic Settings")
 	bool bAllocateHistoryPerPane;
 
+	/*
+	* If true, auto-exposure is generated from a separate low resolution cube render pass, allowing it to be used with panoramic rendering.
+	* Requires "Allocate History Per Pane".
+	*/
+	UPROPERTY(EditAnywhere, BlueprintReadWrite, AdvancedDisplay, Category = "Panoramic Settings", meta = (editcondition = "bAllocateHistoryPerPane"))
+	bool bAutoExposureRenderPass;
+
 protected:
 	TSharedPtr<FAccumulatorPool, ESPMode::ThreadSafe> AccumulatorPool;
 
 	// ToDo: One per high-res tile per pano-pane?
 	TArray<FSceneViewStateReference> OptionalPaneViewStates;
 
+	// When using an auto exposure render pass, holds view states for 6 cube faces
+	TArray<FSceneViewStateReference> AutoExposureViewStates;
+
 	/** The lifetime of this SceneViewExtension is only during the rendering process. It is destroyed as part of TearDown. */
 	TSharedPtr<FOpenColorIODisplayExtension, ESPMode::ThreadSafe> OCIOSceneViewExtension;
 
diff --git a/Engine/Plugins/MovieScene/MovieRenderPipeline/Source/MovieRenderPipelineCore/Private/MoviePipelineRendering.cpp b/Engine/Plugins/MovieScene/MovieRenderPipeline/Source/MovieRenderPipelineCore/Private/MoviePipelineRendering.cpp
--- a/Engine/Plugins/MovieScene/MovieRenderPipeline/Source/MovieRenderPipelineCore/Private/MoviePipelineRendering.cpp
+++ b/Engine/Plugins/MovieScene/MovieRenderPipeline/Source/MovieRenderPipelineCore/Private/MoviePipelineRendering.cpp
@@ -438,6 +438,9 @@ void UMoviePipeline::RenderFrame()
 				SampleState.TextureSharpnessBias = HighResSettings->TextureSharpnessBias;
 				SampleState.OCIOConfiguration = ColorSettings ? &ColorSettings->OCIOConfiguration : nullptr;
 				SampleState.GlobalScreenPercentageFraction = FLegacyScreenPercentageDriver::GetCVarResolutionFraction();
+				SampleState.bAutoExposureRenderPass = HighResSettings->bAutoExposureRenderPass && HighResSettings->TileCount > 1;
+				SampleState.bAutoExposureCubePass = false;
+				SampleState.AutoExposureCubeFace = 0;
 				SampleState.bOverrideCameraOverscan = CameraSettings->bOverrideCameraOverscan;
 				SampleState.OverscanPercentage = FMath::Clamp(CameraSettings->OverscanPercentage, 0.0f, 1.0f);
 
diff --git a/Engine/Plugins/MovieScene/MovieRenderPipeline/Source/MovieRenderPipelineCore/Public/MoviePipelineHighResSetting.h b/Engine/Plugins/MovieScene/MovieRenderPipeline/Source/MovieRenderPipelineCore/Public/MoviePipelineHighResSetting.h
--- a/Engine/Plugins/MovieScene/MovieRenderPipeline/Source/MovieRenderPipelineCore/Public/MoviePipelineHighResSetting.h
+++ b/Engine/Plugins/MovieScene/MovieRenderPipeline/Source/MovieRenderPipelineCore/Public/MoviePipelineHighResSetting.h
@@ -23,6 +23,7 @@ public:
 		, bOverrideSubSurfaceScattering(false)
 		, BurleySampleCount(64)
 		, bAllocateHistoryPerTile(true)
+		, bAutoExposureRenderPass(true)
 	{
 	}
 	
@@ -46,11 +47,15 @@ public:
 
 		if (TileCount > 1)
 		{
-			ValidationResults.Add(NSLOCTEXT("MovieRenderPipeline", "HighRes_UnsupportedFeatures", "Tiling does not support all rendering features (bloom, some screen-space effects). Additionally, TAA and Auto Exposure are not supported and will be forced off when rendering. Use Spatial/Temporal sampling and Manual Camera Exposure /w Exposure Compensation instead. Tiling does not support reference motion blur with the Path Tracer."));
+			ValidationResults.Add(NSLOCTEXT("MovieRenderPipeline", "HighRes_UnsupportedFeatures", "Tiling does not support all rendering features (bloom, some screen-space effects). Additionally, TAA is not supported and will be forced off when rendering. Use Spatial/Temporal sampling instead. Tiling does not support reference motion blur with the Path Tracer."));
 			if (FMath::IsNearlyEqual(OverlapRatio, 0.f))
 			{
 				ValidationResults.Add(NSLOCTEXT("MovieRenderPipeline", "HighRes_OverlapNeeded", "Increase the overlap amount to avoid seams in the final image. More overlap results in longer renders so start at 0.1 and increase as needed."));
 			}
+			if (!bAutoExposureRenderPass)
+			{
+				ValidationResults.Add(NSLOCTEXT("MovieRenderPipeline", "HighRes_ManualExposureNeeded", "Use Manual Camera Exposure /w Exposure Compensation when the Auto Exposure Render Pass is disabled."));
+			}
 			ValidationState = EMoviePipelineValidationState::Warnings;
 		}
 	}
@@ -207,6 +212,12 @@ public:
 	UPROPERTY(EditAnywhere, BlueprintReadWrite, Category = "Render Settings")
 	bool bAllocateHistoryPerTile;
 
+	/*
+	* If true, auto-exposure is generated from a separate low resolution full screen pass, allowing it to be used with tiling.
+	*/
+	UPROPERTY(EditAnywhere, BlueprintReadWrite, Category = "Render Settings")
+	bool bAutoExposureRenderPass;
+
 
 private:
 	/** What was the Burley CVar before we overrode it? */
diff --git a/Engine/Plugins/MovieScene/MovieRenderPipeline/Source/MovieRenderPipelineCore/Public/MovieRenderPipelineDataTypes.h b/Engine/Plugins/MovieScene/MovieRenderPipeline/Source/MovieRenderPipelineCore/Public/MovieRenderPipelineDataTypes.h
--- a/Engine/Plugins/MovieScene/MovieRenderPipeline/Source/MovieRenderPipelineCore/Public/MovieRenderPipelineDataTypes.h
+++ b/Engine/Plugins/MovieScene/MovieRenderPipeline/Source/MovieRenderPipelineCore/Public/MovieRenderPipelineDataTypes.h
@@ -1062,6 +1062,21 @@ public:
 	*/
 	float GlobalScreenPercentageFraction;
 
+	/**
+	* Whether the auto exposure render pass is enabled in the settings.  Needs to be set on ALL passes, so the other passes know to look for the results from the exposure pass.
+	* The exposure rendering itself runs with the first tile.
+	*/
+	bool bAutoExposureRenderPass;
+
+	/**
+	* Whether the auto exposure pass is an omnidirectional cube capture pass, useful for panoramic rendering.  Only set for the actual exposure pass!
+	*/
+	bool bAutoExposureCubePass;
+
+	/**
+	* Transient value set when generating an autoexposure cube face.
+	*/
+	int8 AutoExposureCubeFace;
 
 	FMoviePipelineFrameOutputState OutputState;
 
diff --git a/Engine/Plugins/MovieScene/MovieRenderPipeline/Source/MovieRenderPipelineRenderPasses/Private/MoviePipelineDeferredPasses.cpp b/Engine/Plugins/MovieScene/MovieRenderPipeline/Source/MovieRenderPipelineRenderPasses/Private/MoviePipelineDeferredPasses.cpp
--- a/Engine/Plugins/MovieScene/MovieRenderPipeline/Source/MovieRenderPipelineRenderPasses/Private/MoviePipelineDeferredPasses.cpp
+++ b/Engine/Plugins/MovieScene/MovieRenderPipeline/Source/MovieRenderPipelineRenderPasses/Private/MoviePipelineDeferredPasses.cpp
@@ -245,6 +245,14 @@ void UMoviePipelineDeferredPassBase::SetupImpl(const MoviePipeline::FMoviePipeli
 		// We don't always want to allocate a unique history per tile as very large resolutions can OOM the GPU in backbuffer images alone.
 		// But we do need the history for some features (like Lumen) to work, so it's optional.
 		int32 HighResTileCount = HighResSettings->bAllocateHistoryPerTile ? HighResSettings->TileCount : 1;
+
+		if (HighResSettings->bAutoExposureRenderPass && HighResTileCount > 1 && bRenderMainPass)
+		{
+			// Add view state for auto-exposure render pass
+			FMultiCameraViewStateData::FPerTile& PerTile = CameraData.TileData.FindOrAdd(FIntPoint(-1, -1));
+			PerTile.SceneViewStates.AddDefaulted();
+		}
+
 		for (int32 TileIndexX = 0; TileIndexX < HighResTileCount; TileIndexX++)
 		{
 			for (int32 TileIndexY = 0; TileIndexY < HighResTileCount; TileIndexY++)
@@ -447,6 +455,26 @@ FSceneViewStateInterface* UMoviePipelineDeferredPassBase::GetSceneViewStateInter
 	return nullptr;
 }
 
+FSceneViewStateInterface* UMoviePipelineDeferredPassBase::GetExposureSceneViewStateInterface(IViewCalcPayload* OptPayload, int32 CubeFaceIndex)
+{
+	UE::MoviePipeline::FDeferredPassRenderStatePayload* Payload = (UE::MoviePipeline::FDeferredPassRenderStatePayload*)OptPayload;
+	check(Payload);
+
+	// When not using multi-camera support, CameraIndex is -1, but means "first entry" in the array.
+	int32 LocalCameraIndex = FMath::Clamp(Payload->CameraIndex, 0, Payload->CameraIndex);
+
+	// Special tile {-1,-1} holds the auto-exposure view state, if present
+	FMultiCameraViewStateData& CameraData = CameraViewStateData[LocalCameraIndex];
+	if (FMultiCameraViewStateData::FPerTile* TileData = CameraData.TileData.Find(FIntPoint(-1,-1)))
+	{
+		return TileData->SceneViewStates[0].GetReference();
+	}
+	else
+	{
+		return nullptr;
+	}
+}
+
 void UMoviePipelineDeferredPassBase::GatherOutputPassesImpl(TArray<FMoviePipelinePassIdentifier>& ExpectedRenderPasses)
 {
 	// No super call here because multiple cameras makes this all complicated
@@ -533,73 +561,98 @@ void UMoviePipelineDeferredPassBase::RenderSample_GameThreadImpl(const FMoviePip
 		// Main Render Pass
 		if (bRenderMainPass)
 		{
-			FMoviePipelineRenderPassMetrics InOutSampleState = GetRenderPassMetricsForCamera(CameraIndexForRenderPass, InSampleState);
-			// InOutSampleState.OutputState.CameraCount = NumCameras;
-			InOutSampleState.OutputState.CameraIndex = CameraIndexForRenderPass;
-			InOutSampleState.OutputState.CameraNameOverride = GetCameraNameOverride(CameraIndex);
+			// Auto-exposure pass is rendered only for the first tile and sample, actual pass is always rendered
+			bool bAutoExposurePass = InSampleState.bAutoExposureRenderPass && InSampleState.TileIndexes.X == 0 && InSampleState.TileIndexes.Y == 0 && InSampleState.TemporalSampleIndex == 0 && InSampleState.SpatialSampleIndex == 0;
+			bool bActualPass = true;
+
+			while (bAutoExposurePass || bActualPass)
+			{
+				FMoviePipelineRenderPassMetrics InOutSampleState = GetRenderPassMetricsForCamera(CameraIndexForRenderPass, InSampleState);
+				// InOutSampleState.OutputState.CameraCount = NumCameras;
+				InOutSampleState.OutputState.CameraIndex = CameraIndexForRenderPass;
+				InOutSampleState.OutputState.CameraNameOverride = GetCameraNameOverride(CameraIndex);
+
+				if (bAutoExposurePass)
+				{
+					// Auto exposure reference pass uses tile indices -1,-1
+					InOutSampleState.TileIndexes = FIntPoint(-1, -1);
+				}
 
-			UE::MoviePipeline::FDeferredPassRenderStatePayload Payload;
-			Payload.CameraIndex = CameraIndexForRenderPass;
-			Payload.TileIndex = InOutSampleState.TileIndexes;
+				UE::MoviePipeline::FDeferredPassRenderStatePayload Payload;
+				Payload.CameraIndex = CameraIndexForRenderPass;
+				Payload.TileIndex = InOutSampleState.TileIndexes;
 
-			// Main renders use index 0.
-			Payload.SceneViewIndex = 0;
+				// Main renders use index 0.
+				Payload.SceneViewIndex = 0;
 
-			TSharedPtr<FSceneViewFamilyContext> ViewFamily = CalculateViewFamily(InOutSampleState, &Payload);
+				TSharedPtr<FSceneViewFamilyContext> ViewFamily = CalculateViewFamily(InOutSampleState, &Payload);
 
-			// Add post-processing materials if needed
-			FSceneView* View = const_cast<FSceneView*>(ViewFamily->Views[0]);
-			View->FinalPostProcessSettings.BufferVisualizationOverviewMaterials.Empty();
-			View->FinalPostProcessSettings.BufferVisualizationPipes.Empty();
+				// Add post-processing materials if needed
+				FSceneView* View = const_cast<FSceneView*>(ViewFamily->Views[0]);
+				View->FinalPostProcessSettings.BufferVisualizationOverviewMaterials.Empty();
+				View->FinalPostProcessSettings.BufferVisualizationPipes.Empty();
 
-			for (UMaterialInterface* Material : ActivePostProcessMaterials)
-			{
-				if (Material)
+				for (UMaterialInterface* Material : ActivePostProcessMaterials)
 				{
-					View->FinalPostProcessSettings.BufferVisualizationOverviewMaterials.Add(Material);
+					if (Material)
+					{
+						View->FinalPostProcessSettings.BufferVisualizationOverviewMaterials.Add(Material);
+					}
 				}
-			}
 
-			for (UMaterialInterface* VisMaterial : View->FinalPostProcessSettings.BufferVisualizationOverviewMaterials)
-			{
-				// If this was just to contribute to the history buffer, no need to go any further.
-				if (InOutSampleState.bDiscardResult)
+				for (UMaterialInterface* VisMaterial : View->FinalPostProcessSettings.BufferVisualizationOverviewMaterials)
 				{
-					continue;
-				}
+					// If this was just to contribute to the history buffer, no need to go any further.
+					if (InOutSampleState.bDiscardResult)
+					{
+						continue;
+					}
 
-				FMoviePipelinePassIdentifier LayerPassIdentifier = FMoviePipelinePassIdentifier(PassIdentifier.Name + GetNameForPostProcessMaterial(VisMaterial), PassIdentifierForCurrentCamera.CameraName);
+					FMoviePipelinePassIdentifier LayerPassIdentifier = FMoviePipelinePassIdentifier(PassIdentifier.Name + GetNameForPostProcessMaterial(VisMaterial), PassIdentifierForCurrentCamera.CameraName);
 
-				auto BufferPipe = MakeShared<FImagePixelPipe, ESPMode::ThreadSafe>();
-				BufferPipe->bIsExpecting32BitPixelData = ActiveHighPrecisionPostProcessMaterials.Contains(VisMaterial);
-				BufferPipe->AddEndpoint(MakeForwardingEndpoint(LayerPassIdentifier, InOutSampleState));
+					auto BufferPipe = MakeShared<FImagePixelPipe, ESPMode::ThreadSafe>();
+					BufferPipe->bIsExpecting32BitPixelData = ActiveHighPrecisionPostProcessMaterials.Contains(VisMaterial);
+					BufferPipe->AddEndpoint(MakeForwardingEndpoint(LayerPassIdentifier, InOutSampleState));
 
-				View->FinalPostProcessSettings.BufferVisualizationPipes.Add(VisMaterial->GetFName(), BufferPipe);
-			}
+					View->FinalPostProcessSettings.BufferVisualizationPipes.Add(VisMaterial->GetFName(), BufferPipe);
+				}
 
 
-			int32 NumValidMaterials = View->FinalPostProcessSettings.BufferVisualizationPipes.Num();
-			View->FinalPostProcessSettings.bBufferVisualizationDumpRequired = NumValidMaterials > 0;
+				int32 NumValidMaterials = View->FinalPostProcessSettings.BufferVisualizationPipes.Num();
+				View->FinalPostProcessSettings.bBufferVisualizationDumpRequired = NumValidMaterials > 0;
 
-			// Submit to be rendered. Main render pass always uses target 0.
-			TWeakObjectPtr<UTextureRenderTarget2D> ViewRenderTarget = GetOrCreateViewRenderTarget(InOutSampleState.BackbufferSize, (IViewCalcPayload*)(&Payload));
-			check(ViewRenderTarget.IsValid());
+				// Submit to be rendered. Main render pass always uses target 0.
+				TWeakObjectPtr<UTextureRenderTarget2D> ViewRenderTarget = GetOrCreateViewRenderTarget(InOutSampleState.BackbufferSize, (IViewCalcPayload*)(&Payload));
+				check(ViewRenderTarget.IsValid());
 
-			FRenderTarget* RenderTarget = ViewRenderTarget->GameThread_GetRenderTargetResource();
-			check(RenderTarget);
+				FRenderTarget* RenderTarget = ViewRenderTarget->GameThread_GetRenderTargetResource();
+				check(RenderTarget);
 
-			FCanvas Canvas = FCanvas(RenderTarget, nullptr, GetPipeline()->GetWorld(), View->GetFeatureLevel(), FCanvas::CDM_DeferDrawing, 1.0f);
-			GetRendererModule().BeginRenderingViewFamily(&Canvas, ViewFamily.Get());
+				FCanvas Canvas = FCanvas(RenderTarget, nullptr, GetPipeline()->GetWorld(), View->GetFeatureLevel(), FCanvas::CDM_DeferDrawing, 1.0f);
+				GetRendererModule().BeginRenderingViewFamily(&Canvas, ViewFamily.Get());
 			
-			ENQUEUE_RENDER_COMMAND(TransitionTextureSRVState)(
-			[RenderTarget](FRHICommandListImmediate& RHICmdList) mutable
-			{
-				// Transition our render target from a render target view to a shader resource view to allow the UMG preview material to read from this Render Target.
-				RHICmdList.Transition(FRHITransitionInfo(RenderTarget->GetRenderTargetTexture(), ERHIAccess::RTV, ERHIAccess::SRVGraphicsPixel));
-			});
+				ENQUEUE_RENDER_COMMAND(TransitionTextureSRVState)(
+				[RenderTarget](FRHICommandListImmediate& RHICmdList) mutable
+				{
+					// Transition our render target from a render target view to a shader resource view to allow the UMG preview material to read from this Render Target.
+					RHICmdList.Transition(FRHITransitionInfo(RenderTarget->GetRenderTargetTexture(), ERHIAccess::RTV, ERHIAccess::SRVGraphicsPixel));
+				});
+
+				if (bAutoExposurePass)
+				{
+					// Finished rendering auto-exposure pass
+					bAutoExposurePass = false;
+				}
+				else
+				{
+					// Readback + Accumulate.
+					PostRendererSubmission(InOutSampleState, PassIdentifierForCurrentCamera, GetOutputFileSortingOrder(), Canvas);
+
+					// Finished rendering actual pass
+					bActualPass = false;
+				}
+			}
 
-			// Readback + Accumulate.
-			PostRendererSubmission(InOutSampleState, PassIdentifierForCurrentCamera, GetOutputFileSortingOrder(), Canvas);
 		}
 
 
@@ -1038,8 +1091,8 @@ void UMoviePipelineDeferredPassBase::PostRendererSubmission(const FMoviePipeline
 
 bool UMoviePipelineDeferredPassBase::IsAutoExposureAllowed(const FMoviePipelineRenderPassMetrics& InSampleState) const
 {
-	// High-res tiling doesn't support auto-exposure.
-	return !(InSampleState.GetTileCount() > 1);
+	// High-res tiling supports auto-exposure, if the auto exposure render pass is enabled.
+	return !(InSampleState.GetTileCount() > 1) || InSampleState.bAutoExposureRenderPass;
 }
 
 #if WITH_EDITOR
diff --git a/Engine/Plugins/MovieScene/MovieRenderPipeline/Source/MovieRenderPipelineRenderPasses/Private/MoviePipelineImagePassBase.cpp b/Engine/Plugins/MovieScene/MovieRenderPipeline/Source/MovieRenderPipelineRenderPasses/Private/MoviePipelineImagePassBase.cpp
--- a/Engine/Plugins/MovieScene/MovieRenderPipeline/Source/MovieRenderPipelineRenderPasses/Private/MoviePipelineImagePassBase.cpp
+++ b/Engine/Plugins/MovieScene/MovieRenderPipeline/Source/MovieRenderPipelineRenderPasses/Private/MoviePipelineImagePassBase.cpp
@@ -226,8 +226,20 @@ TSharedPtr<FSceneViewFamilyContext> UMoviePipelineImagePassBase::CalculateViewFa
 		}
 	}
 
+	// Auto exposure pass is specified with a tile index of {-1,-1}
+	const bool bAutoExposurePass = (InOutSampleState.TileIndexes.X == -1) && (InOutSampleState.TileIndexes.Y == -1);
+	const bool bScreenPercentageSupported = IsScreenPercentageSupported() && !bAutoExposurePass;
+
+	// Force disable screen percentage and motion blur for auto-exposure passes.  These are already at lowered resolution relative to the overall high res
+	// tiled view, and only the eye adaptation is used from these, making blur irrelevant.  Saves history memory and performance.
+	if (bAutoExposurePass)
+	{
+		OutViewFamily->EngineShowFlags.ScreenPercentage = false;
+		OutViewFamily->EngineShowFlags.MotionBlur = false;
+	}
+
 	// No need to do anything if screen percentage is not supported. 
-	if (IsScreenPercentageSupported())
+	if (bScreenPercentageSupported)
 	{
 		// Allows all Output Settings to have an access to View Family. This allows to modify rendering output settings.
 		for (UMoviePipelineViewFamilySetting* Setting : GetPipeline()->FindSettingsForShot<UMoviePipelineViewFamilySetting>(Shot))
@@ -239,55 +251,61 @@ TSharedPtr<FSceneViewFamilyContext> UMoviePipelineImagePassBase::CalculateViewFa
 	// If UMoviePipelineViewFamilySetting never set a Screen percentage interface we fallback to default.
 	if (OutViewFamily->GetScreenPercentageInterface() == nullptr)
 	{
-		OutViewFamily->SetScreenPercentageInterface(new FLegacyScreenPercentageDriver(*OutViewFamily, IsScreenPercentageSupported() ? InOutSampleState.GlobalScreenPercentageFraction : 1.f));
+		OutViewFamily->SetScreenPercentageInterface(new FLegacyScreenPercentageDriver(*OutViewFamily, bScreenPercentageSupported ? InOutSampleState.GlobalScreenPercentageFraction : 1.f));
 	}
 
+	int32 ViewCount = InOutSampleState.bAutoExposureCubePass ? 6 : 1;
+	for (int32 ViewIndex = 0; ViewIndex < ViewCount; ViewIndex++)
+	{
+		// Ignored in downstream code if this isn't an auto exposure cube pass
+		InOutSampleState.AutoExposureCubeFace = ViewIndex;
 
-	// View is added as a child of the OutViewFamily-> 
-	FSceneView* View = GetSceneViewForSampleState(OutViewFamily.Get(), /*InOut*/ InOutSampleState, OptPayload);
-	
-	SetupViewForViewModeOverride(View);
-
-	// Override the view's FrameIndex to be based on our progress through the sequence. This greatly increases
-	// determinism with things like TAA.
-	View->OverrideFrameIndexValue = InOutSampleState.FrameIndex;
-	View->OverrideOutputFrameIndexValue = InOutSampleState.OutputState.OutputFrameNumber;
-	View->bCameraCut = InOutSampleState.bCameraCut;
-	View->bIsOfflineRender = true;
-	View->AntiAliasingMethod = IsAntiAliasingSupported() ? InOutSampleState.AntiAliasingMethod : EAntiAliasingMethod::AAM_None;
+		// View is added as a child of the OutViewFamily-> 
+		FSceneView* View = GetSceneViewForSampleState(OutViewFamily.Get(), /*InOut*/ InOutSampleState, OptPayload);
 
-	// Override the Motion Blur settings since these are controlled by the movie pipeline.
-	{
-		FFrameRate OutputFrameRate = GetPipeline()->GetPipelinePrimaryConfig()->GetEffectiveFrameRate(GetPipeline()->GetTargetSequence());
+		SetupViewForViewModeOverride(View);
 
-		// We need to inversly scale the target FPS by time dilation to counteract slowmo. If scaling isn't applied then motion blur length
-		// stays the same length despite the smaller delta time and the blur ends up too long.
-		View->FinalPostProcessSettings.MotionBlurTargetFPS = FMath::RoundToInt(OutputFrameRate.AsDecimal() / FMath::Max(SMALL_NUMBER, InOutSampleState.OutputState.TimeData.TimeDilation));
-		View->FinalPostProcessSettings.MotionBlurAmount = InOutSampleState.OutputState.TimeData.MotionBlurFraction;
-		View->FinalPostProcessSettings.MotionBlurMax = 100.f;
-		View->FinalPostProcessSettings.bOverride_MotionBlurAmount = true;
-		View->FinalPostProcessSettings.bOverride_MotionBlurTargetFPS = true;
-		View->FinalPostProcessSettings.bOverride_MotionBlurMax = true;
+		// Override the view's FrameIndex to be based on our progress through the sequence. This greatly increases
+		// determinism with things like TAA.
+		View->OverrideFrameIndexValue = InOutSampleState.FrameIndex;
+		View->OverrideOutputFrameIndexValue = InOutSampleState.OutputState.OutputFrameNumber;
+		View->bCameraCut = InOutSampleState.bCameraCut;
+		View->bIsOfflineRender = true;
+		View->AntiAliasingMethod = IsAntiAliasingSupported() ? InOutSampleState.AntiAliasingMethod : EAntiAliasingMethod::AAM_None;
 
-		// Skip the whole pass if they don't want motion blur.
-		if (FMath::IsNearlyZero(InOutSampleState.OutputState.TimeData.MotionBlurFraction))
+		// Override the Motion Blur settings since these are controlled by the movie pipeline.
 		{
-			OutViewFamily->EngineShowFlags.SetMotionBlur(false);
+			FFrameRate OutputFrameRate = GetPipeline()->GetPipelinePrimaryConfig()->GetEffectiveFrameRate(GetPipeline()->GetTargetSequence());
+
+			// We need to inversly scale the target FPS by time dilation to counteract slowmo. If scaling isn't applied then motion blur length
+			// stays the same length despite the smaller delta time and the blur ends up too long.
+			View->FinalPostProcessSettings.MotionBlurTargetFPS = FMath::RoundToInt(OutputFrameRate.AsDecimal() / FMath::Max(SMALL_NUMBER, InOutSampleState.OutputState.TimeData.TimeDilation));
+			View->FinalPostProcessSettings.MotionBlurAmount = InOutSampleState.OutputState.TimeData.MotionBlurFraction;
+			View->FinalPostProcessSettings.MotionBlurMax = 100.f;
+			View->FinalPostProcessSettings.bOverride_MotionBlurAmount = true;
+			View->FinalPostProcessSettings.bOverride_MotionBlurTargetFPS = true;
+			View->FinalPostProcessSettings.bOverride_MotionBlurMax = true;
+
+			// Skip the whole pass if they don't want motion blur.
+			if (FMath::IsNearlyZero(InOutSampleState.OutputState.TimeData.MotionBlurFraction))
+			{
+				OutViewFamily->EngineShowFlags.SetMotionBlur(false);
+			}
 		}
-	}
 
-	// Locked Exposure
-	const bool bAutoExposureAllowed = IsAutoExposureAllowed(InOutSampleState);
-	{
-		// If the rendering pass doesn't allow autoexposure and they dont' have manual exposure set up, warn.
-		if (!bAutoExposureAllowed && (View->FinalPostProcessSettings.AutoExposureMethod != EAutoExposureMethod::AEM_Manual))
+		// Locked Exposure
+		const bool bAutoExposureAllowed = IsAutoExposureAllowed(InOutSampleState);
 		{
-			// Skip warning if the project setting is disabled though, as exposure will be forced off in the renderer anyways.
-			const URendererSettings* RenderSettings = GetDefault<URendererSettings>();
-			if (RenderSettings->bDefaultFeatureAutoExposure != false)
+			// If the rendering pass doesn't allow autoexposure and they dont' have manual exposure set up, warn.
+			if (!bAutoExposureAllowed && (View->FinalPostProcessSettings.AutoExposureMethod != EAutoExposureMethod::AEM_Manual))
 			{
-				UE_LOG(LogMovieRenderPipeline, Warning, TEXT("Camera Auto Exposure Method not supported by one or more render passes. Change the Auto Exposure Method to Manual!"));
-				View->FinalPostProcessSettings.AutoExposureMethod = EAutoExposureMethod::AEM_Manual;
+				// Skip warning if the project setting is disabled though, as exposure will be forced off in the renderer anyways.
+				const URendererSettings* RenderSettings = GetDefault<URendererSettings>();
+				if (RenderSettings->bDefaultFeatureAutoExposure != false)
+				{
+					UE_LOG(LogMovieRenderPipeline, Warning, TEXT("Camera Auto Exposure Method not supported by one or more render passes. Change the Auto Exposure Method to Manual!"));
+					View->FinalPostProcessSettings.AutoExposureMethod = EAutoExposureMethod::AEM_Manual;
+				}
 			}
 		}
 	}
@@ -301,100 +319,106 @@ TSharedPtr<FSceneViewFamilyContext> UMoviePipelineImagePassBase::CalculateViewFa
 		ViewExt->SetupViewFamily(*OutViewFamily.Get());
 	}
 
-	for (int ViewExt = 0; ViewExt < OutViewFamily->ViewExtensions.Num(); ViewExt++)
-	{
-		OutViewFamily->ViewExtensions[ViewExt]->SetupView(*OutViewFamily.Get(), *View);
-	}
-
 	// Support scene captures with the "bMainViewFamily" flag set
 	OutViewFamily->bIsMainViewFamily = true;
 
-	// The requested configuration may not be supported, warn user and fall back. We can't call
-	// FSceneView::SetupAntiAliasingMethod because it reads the value from the cvar which would
-	// cause the value set by the MoviePipeline UI to be ignored.
+	// Post view family extension setup, do some more work on each view
+	for (int32 ViewIndex = 0; ViewIndex < ViewCount; ViewIndex++)
 	{
-		bool bMethodWasUnsupported = false;
-		if (View->AntiAliasingMethod == AAM_TemporalAA && !SupportsGen4TAA(View->GetShaderPlatform()))
-		{
-			UE_LOG(LogMovieRenderPipeline, Error, TEXT("TAA was requested but this hardware does not support it."));
-			bMethodWasUnsupported = true;
-		}
-		else if (View->AntiAliasingMethod == AAM_TSR && !SupportsTSR(View->GetShaderPlatform()))
+		FSceneView* View = const_cast<FSceneView*>(OutViewFamily->Views[ViewIndex]);
+
+		for (int ViewExt = 0; ViewExt < OutViewFamily->ViewExtensions.Num(); ViewExt++)
 		{
-			UE_LOG(LogMovieRenderPipeline, Error, TEXT("TSR was requested but this hardware does not support it."));
-			bMethodWasUnsupported = true;
+			OutViewFamily->ViewExtensions[ViewExt]->SetupView(*OutViewFamily.Get(), *View);
 		}
 
-		if (bMethodWasUnsupported)
+		// The requested configuration may not be supported, warn user and fall back. We can't call
+		// FSceneView::SetupAntiAliasingMethod because it reads the value from the cvar which would
+		// cause the value set by the MoviePipeline UI to be ignored.
 		{
-			View->AntiAliasingMethod = AAM_None;
+			bool bMethodWasUnsupported = false;
+			if (View->AntiAliasingMethod == AAM_TemporalAA && !SupportsGen4TAA(View->GetShaderPlatform()))
+			{
+				UE_LOG(LogMovieRenderPipeline, Error, TEXT("TAA was requested but this hardware does not support it."));
+				bMethodWasUnsupported = true;
+			}
+			else if (View->AntiAliasingMethod == AAM_TSR && !SupportsTSR(View->GetShaderPlatform()))
+			{
+				UE_LOG(LogMovieRenderPipeline, Error, TEXT("TSR was requested but this hardware does not support it."));
+				bMethodWasUnsupported = true;
+			}
+
+			if (bMethodWasUnsupported)
+			{
+				View->AntiAliasingMethod = AAM_None;
+			}
 		}
-	}
 
-	// Anti Aliasing
-	{
-		// If we're not using Temporal Anti-Aliasing or Path Tracing we will apply the View Matrix projection jitter. Normally TAA sets this
-		// inside FSceneRenderer::PreVisibilityFrameSetup. Path Tracing does its own anti-aliasing internally.
-		bool bApplyProjectionJitter = 
-									!OutViewFamily->EngineShowFlags.PathTracing 
-									&& !IsTemporalAccumulationBasedMethod(View->AntiAliasingMethod);
-		if (bApplyProjectionJitter)
+		// Anti Aliasing
 		{
-			View->ViewMatrices.HackAddTemporalAAProjectionJitter(InOutSampleState.ProjectionMatrixJitterAmount);
+			// If we're not using Temporal Anti-Aliasing or Path Tracing we will apply the View Matrix projection jitter. Normally TAA sets this
+			// inside FSceneRenderer::PreVisibilityFrameSetup. Path Tracing does its own anti-aliasing internally.
+			bool bApplyProjectionJitter =
+				!OutViewFamily->EngineShowFlags.PathTracing
+				&& !IsTemporalAccumulationBasedMethod(View->AntiAliasingMethod);
+			if (bApplyProjectionJitter)
+			{
+				View->ViewMatrices.HackAddTemporalAAProjectionJitter(InOutSampleState.ProjectionMatrixJitterAmount);
+			}
 		}
-	}
 
-	// Path Tracer Sampling
-	if (OutViewFamily->EngineShowFlags.PathTracing)
-	{
-		// override whatever settings came from PostProcessVolume or Camera
+		// Path Tracer Sampling
+		if (OutViewFamily->EngineShowFlags.PathTracing)
+		{
+			// override whatever settings came from PostProcessVolume or Camera
 
-		// If motion blur is enabled:
-		//    blend all spatial samples together while leaving the handling of temporal samples up to MRQ
-		//    each temporal sample will include denoising and post-process effects
-		// If motion blur is NOT enabled:
-		//    blend all temporal+spatial samples within the path tracer and only apply denoising on the last temporal sample
-		//    this way we minimize denoising cost and also allow a much higher number of temporal samples to be used which
-		//    can help reduce strobing
+			// If motion blur is enabled:
+			//    blend all spatial samples together while leaving the handling of temporal samples up to MRQ
+			//    each temporal sample will include denoising and post-process effects
+			// If motion blur is NOT enabled:
+			//    blend all temporal+spatial samples within the path tracer and only apply denoising on the last temporal sample
+			//    this way we minimize denoising cost and also allow a much higher number of temporal samples to be used which
+			//    can help reduce strobing
 
-		// NOTE: Tiling is not compatible with the reference motion blur mode because it changes the order of the loops over the image.
-		const bool bAccumulateSpatialSamplesOnly = OutViewFamily->EngineShowFlags.MotionBlur || InOutSampleState.GetTileCount() > 1;
+			// NOTE: Tiling is not compatible with the reference motion blur mode because it changes the order of the loops over the image.
+			const bool bAccumulateSpatialSamplesOnly = OutViewFamily->EngineShowFlags.MotionBlur || InOutSampleState.GetTileCount() > 1;
 
-		const int32 SampleCount = bAccumulateSpatialSamplesOnly ? InOutSampleState.SpatialSampleCount : InOutSampleState.TemporalSampleCount * InOutSampleState.SpatialSampleCount;
-		const int32 SampleIndex = bAccumulateSpatialSamplesOnly ? InOutSampleState.SpatialSampleIndex : InOutSampleState.TemporalSampleIndex * InOutSampleState.SpatialSampleCount + InOutSampleState.SpatialSampleIndex;
+			const int32 SampleCount = bAccumulateSpatialSamplesOnly ? InOutSampleState.SpatialSampleCount : InOutSampleState.TemporalSampleCount * InOutSampleState.SpatialSampleCount;
+			const int32 SampleIndex = bAccumulateSpatialSamplesOnly ? InOutSampleState.SpatialSampleIndex : InOutSampleState.TemporalSampleIndex * InOutSampleState.SpatialSampleCount + InOutSampleState.SpatialSampleIndex;
 
-		// TODO: pass along FrameIndex (which includes SampleIndex) to make sure sampling is fully deterministic
+			// TODO: pass along FrameIndex (which includes SampleIndex) to make sure sampling is fully deterministic
 
-		// Overwrite whatever sampling count came from the PostProcessVolume
-		View->FinalPostProcessSettings.bOverride_PathTracingSamplesPerPixel = true;
-		View->FinalPostProcessSettings.PathTracingSamplesPerPixel = SampleCount;
+			// Overwrite whatever sampling count came from the PostProcessVolume
+			View->FinalPostProcessSettings.bOverride_PathTracingSamplesPerPixel = true;
+			View->FinalPostProcessSettings.PathTracingSamplesPerPixel = SampleCount;
 
-		// reset path tracer's accumulation at the start of each sample
-		View->bForcePathTracerReset = SampleIndex == 0;
+			// reset path tracer's accumulation at the start of each sample
+			View->bForcePathTracerReset = SampleIndex == 0;
 
-		// discard the result, unless its the last sample
-		InOutSampleState.bDiscardResult |= !(SampleIndex == SampleCount - 1);
-	}
+			// discard the result, unless its the last sample
+			InOutSampleState.bDiscardResult |= !(SampleIndex == SampleCount - 1);
+		}
 
-	// Object Occlusion/Histories
-	{
-		// If we're using tiling, we force the reset of histories each frame so that we don't use the previous tile's
-		// object occlusion queries, as that causes things to disappear from some views.
-		if (InOutSampleState.GetTileCount() > 1)
+		// Object Occlusion/Histories
 		{
-			View->bForceCameraVisibilityReset = true;
+			// If we're using tiling, we force the reset of histories each frame so that we don't use the previous tile's
+			// object occlusion queries, as that causes things to disappear from some views.
+			if (InOutSampleState.GetTileCount() > 1)
+			{
+				View->bForceCameraVisibilityReset = true;
+			}
 		}
-	}
 
-	// Bias all mip-mapping to pretend to be working at our target resolution and not our tile resolution
-	// so that the images don't end up soft.
-	{
-		float EffectivePrimaryResolutionFraction = 1.f / InOutSampleState.TileCounts.X;
-		View->MaterialTextureMipBias = FMath::Log2(EffectivePrimaryResolutionFraction);
+		// Bias all mip-mapping to pretend to be working at our target resolution and not our tile resolution
+		// so that the images don't end up soft.
+		{
+			float EffectivePrimaryResolutionFraction = 1.f / InOutSampleState.TileCounts.X;
+			View->MaterialTextureMipBias = FMath::Log2(EffectivePrimaryResolutionFraction);
 
-		// Add an additional bias per user settings. This allows them to choose to make the textures sharper if it
-		// looks better with their particular settings.
-		View->MaterialTextureMipBias += InOutSampleState.TextureSharpnessBias;
+			// Add an additional bias per user settings. This allows them to choose to make the textures sharper if it
+			// looks better with their particular settings.
+			View->MaterialTextureMipBias += InOutSampleState.TextureSharpnessBias;
+		}
 	}
 
 	return OutViewFamily;
@@ -432,22 +456,70 @@ void UMoviePipelineImagePassBase::GatherOutputPassesImpl(TArray<FMoviePipelinePa
 	ExpectedRenderPasses.Add(PassIdentifier);
 }
 
+extern int32 ComputeAutoExposureCubeCaptureSize(FIntPoint Resolution);
+
 FSceneView* UMoviePipelineImagePassBase::GetSceneViewForSampleState(FSceneViewFamily* ViewFamily, FMoviePipelineRenderPassMetrics& InOutSampleState, IViewCalcPayload* OptPayload)
 {
 	APlayerController* LocalPlayerController = GetPipeline()->GetWorld()->GetFirstPlayerController();
 
-	int32 TileSizeX = InOutSampleState.BackbufferSize.X;
-	int32 TileSizeY = InOutSampleState.BackbufferSize.Y;
+	int32 TileSizeX;
+	int32 TileSizeY;
+
+	// Auto exposure pass is specified with a tile index of {-1,-1}
+	const bool bAutoExposurePass = (InOutSampleState.TileIndexes.X == -1) && (InOutSampleState.TileIndexes.Y == -1);
+
+	if (bAutoExposurePass)
+	{
+		if (InOutSampleState.bAutoExposureCubePass)
+		{
+			int32 CubeCaptureSize = ComputeAutoExposureCubeCaptureSize(InOutSampleState.BackbufferSize);
+			check(CubeCaptureSize > 0);
+
+			TileSizeX = CubeCaptureSize;
+			TileSizeY = CubeCaptureSize;
+		}
+		else
+		{
+			// Auto exposure pass renders full screen, but at single tile resolution.  Uses the same back buffer size, so it doesn't require separate render targets.
+			// EffectiveOutputResolution is deprecated in favor of OverscannedResolution in all other code paths, but for this specific code path, we want no overscan.
+			PRAGMA_DISABLE_DEPRECATION_WARNINGS
+			TileSizeX = InOutSampleState.EffectiveOutputResolution.X / InOutSampleState.TileCounts.X;
+			TileSizeY = InOutSampleState.EffectiveOutputResolution.Y / InOutSampleState.TileCounts.Y;
+			PRAGMA_ENABLE_DEPRECATION_WARNINGS
+
+			check(TileSizeX <= InOutSampleState.BackbufferSize.X);
+			check(TileSizeY <= InOutSampleState.BackbufferSize.Y);
+		}
+
+		InOutSampleState.OverscanPercentage = 0.0f;
+	}
+	else
+	{
+		TileSizeX = InOutSampleState.BackbufferSize.X;
+		TileSizeY = InOutSampleState.BackbufferSize.Y;
+	}
 
 	UE::MoviePipeline::FImagePassCameraViewData CameraInfo = GetCameraInfo(InOutSampleState, OptPayload);
 
-	const float DestAspectRatio = InOutSampleState.BackbufferSize.X / (float)InOutSampleState.BackbufferSize.Y;
+	const float DestAspectRatio = TileSizeX / (float)TileSizeY;
 	const float CameraAspectRatio = bAllowCameraAspectRatio ? CameraInfo.ViewInfo.AspectRatio : DestAspectRatio;
 
+	// Auto exposure cube map faces are rendered as 3x2 split screen tiles.
+	static const FIntPoint GCubeFaceViewRectOffsets[6] =
+	{
+		{ 0,0 },
+		{ 1,0 },
+		{ 2,0 },
+		{ 0,1 },
+		{ 1,1 },
+		{ 2,1 },
+	};
+	FIntPoint ViewRectOffset = InOutSampleState.bAutoExposureCubePass ? GCubeFaceViewRectOffsets[InOutSampleState.AutoExposureCubeFace] * FIntPoint(TileSizeX, TileSizeY) : FIntPoint(0, 0);
+
 	FSceneViewInitOptions ViewInitOptions;
 	ViewInitOptions.ViewFamily = ViewFamily;
 	ViewInitOptions.ViewOrigin = CameraInfo.ViewInfo.Location;
-	FIntRect ViewRect = FIntRect(FIntPoint(0, 0), FIntPoint(TileSizeX, TileSizeY));
+	FIntRect ViewRect = FIntRect(ViewRectOffset, ViewRectOffset + FIntPoint(TileSizeX, TileSizeY));
 	ViewInitOptions.SetViewRectangle(ViewRect);
 
 	ViewInitOptions.ViewRotationMatrix = FInverseRotationMatrix(CameraInfo.ViewInfo.Rotation);
@@ -459,8 +531,18 @@ FSceneView* UMoviePipelineImagePassBase::GetSceneViewForSampleState(FSceneViewFa
 		FPlane(1, 0, 0, 0),
 		FPlane(0, 1, 0, 0),
 		FPlane(0, 0, 0, 1));
-	
-	if (InOutSampleState.bOverrideCameraOverscan)
+
+	if (InOutSampleState.bAutoExposureCubePass)
+	{
+		ViewInitOptions.ViewRotationMatrix = ViewInitOptions.ViewRotationMatrix * CalcCubeFaceTransform((ECubeFace)InOutSampleState.AutoExposureCubeFace);
+	}
+
+	if (bAutoExposurePass)
+	{
+		// Overscan is irrelevant for the auto exposure pass
+		CameraInfo.ViewInfo.ClearOverscan();
+	}
+	else if (InOutSampleState.bOverrideCameraOverscan)
 	{
 		// If we are overriding the camera's overscan, clear out any overscan the camera added to the view info, and apply the overriding overscan
 		CameraInfo.ViewInfo.ClearOverscan();
@@ -492,13 +574,28 @@ FSceneView* UMoviePipelineImagePassBase::GetSceneViewForSampleState(FSceneViewFa
 
 	float DofSensorScale = 1.0f;
 
-	if (CameraInfo.bUseCustomProjectionMatrix)
+	if (InOutSampleState.bAutoExposureCubePass)
+	{
+		// Auto exposure cube faces just use fixed 90 degree FOV
+		ViewInitOptions.FOV = 90.0f;
+		ViewInitOptions.DesiredFOV = 90.0f;
+
+		const float MatrixFOV = 90.0f * (float)UE_PI / 360.0f;
+		const float ClippingPlane = GNearClippingPlane;
+
+		ViewInitOptions.ProjectionMatrix = FReversedZPerspectiveMatrix(MatrixFOV, MatrixFOV, 1.0f, 1.0f, ClippingPlane, ClippingPlane);
+	}
+	else if (CameraInfo.bUseCustomProjectionMatrix)
 	{
 		ViewInitOptions.ProjectionMatrix = CameraInfo.CustomProjectionMatrix;
 
-		// Modify the custom matrix to do an off center projection, with overlap for high-res tiling
-		const bool bOrthographic = false;
-		ModifyProjectionMatrixForTiling(InOutSampleState, bOrthographic, /*InOut*/ ViewInitOptions.ProjectionMatrix, /*Out*/ DofSensorScale);
+		// Auto exposure pass is full screen, and doesn't use tiling
+		if (!bAutoExposurePass)
+		{
+			// Modify the custom matrix to do an off center projection, with overlap for high-res tiling
+			const bool bOrthographic = false;
+			ModifyProjectionMatrixForTiling(InOutSampleState, bOrthographic, /*InOut*/ ViewInitOptions.ProjectionMatrix, /*Out*/ DofSensorScale);
+		}
 	}
 	else
 	{
@@ -542,7 +639,11 @@ FSceneView* UMoviePipelineImagePassBase::GetSceneViewForSampleState(FSceneViewFa
 		FIntRect ViewExtents = FViewport::CalculateViewExtents(CameraInfo.ViewInfo.AspectRatio, DestAspectRatio, ViewRect, InOutSampleState.BackbufferSize);
 		FMinimalViewInfo::CalculateProjectionMatrixGivenViewRectangle(CameraInfo.ViewInfo, AspectRatioAxisConstraint, ViewExtents, ViewInitOptions);
 
-		ModifyProjectionMatrixForTiling(InOutSampleState, CameraInfo.ViewInfo.ProjectionMode == ECameraProjectionMode::Orthographic,  /*InOut*/ ViewInitOptions.ProjectionMatrix, /*Out*/ DofSensorScale);
+		// Auto exposure pass is full screen, and doesn't use tiling
+		if (!bAutoExposurePass)
+		{
+			ModifyProjectionMatrixForTiling(InOutSampleState, CameraInfo.ViewInfo.ProjectionMode == ECameraProjectionMode::Orthographic,  /*InOut*/ ViewInitOptions.ProjectionMatrix, /*Out*/ DofSensorScale);
+		}
 	}
 
 
@@ -551,6 +652,12 @@ FSceneView* UMoviePipelineImagePassBase::GetSceneViewForSampleState(FSceneViewFa
 
 	ViewInitOptions.SceneViewStateInterface = GetSceneViewStateInterface(OptPayload);
 
+	// If not the auto exposure pass, attempt to get the view state interface from the auto exposure pass
+	if (!bAutoExposurePass)
+	{
+		ViewInitOptions.ExposureSceneViewStateInterface = GetExposureSceneViewStateInterface(OptPayload);
+	}
+
 	FSceneView* View = new FSceneView(ViewInitOptions);
 	ViewFamily->Views.Add(View);
 
@@ -567,8 +674,19 @@ FSceneView* UMoviePipelineImagePassBase::GetSceneViewForSampleState(FSceneViewFa
 	// cause the circle of confusion to be unchanged.
 	View->FinalPostProcessSettings.DepthOfFieldSensorWidth *= DofSensorScale;
 
-	// Modify the 'center' of the lens to be offset for high-res tiling, helps some effects (vignette) etc. still work.
-	View->LensPrincipalPointOffsetScale = (FVector4f)CalculatePrinciplePointOffsetForTiling(InOutSampleState); // LWC_TODO: precision loss. CalculatePrinciplePointOffsetForTiling() could return float, it's normalized?
+	// Disable anti-aliasing and temporal upscale for auto-exposure passes.  Auto-exposure is calculated before those passes, so this is wasted work (and memory for history).
+	if (bAutoExposurePass)
+	{
+		View->AntiAliasingMethod = AAM_None;
+		View->PrimaryScreenPercentageMethod = EPrimaryScreenPercentageMethod::SpatialUpscale;
+	}
+
+	// Auto exposure pass is full screen, and doesn't use tiling
+	if (!bAutoExposurePass)
+	{
+		// Modify the 'center' of the lens to be offset for high-res tiling, helps some effects (vignette) etc. still work.
+		View->LensPrincipalPointOffsetScale = (FVector4f)CalculatePrinciplePointOffsetForTiling(InOutSampleState); // LWC_TODO: precision loss. CalculatePrinciplePointOffsetForTiling() could return float, it's normalized?
+	}
 	View->EndFinalPostprocessSettings(ViewInitOptions);
 
 	// This metadata is per-file and not per-view, but we need the blended result from the view to actually match what we rendered.
@@ -685,6 +803,45 @@ void UMoviePipelineImagePassBase::ModifyProjectionMatrixForTiling(const FMoviePi
 	}
 }
 
+/** Creates a transformation for a cubemap face, following the D3D cubemap layout. */
+FMatrix UMoviePipelineImagePassBase::CalcCubeFaceTransform(ECubeFace Face) const
+{
+	static const FVector XAxis(1.f, 0.f, 0.f);
+	static const FVector YAxis(0.f, 1.f, 0.f);
+	static const FVector ZAxis(0.f, 0.f, 1.f);
+
+	// vectors we will need for our basis
+	FVector vUp(YAxis);
+	FVector vDir;
+	switch (Face)
+	{
+	case CubeFace_PosX:
+		vDir = XAxis;
+		break;
+	case CubeFace_NegX:
+		vDir = -XAxis;
+		break;
+	case CubeFace_PosY:
+		vUp = -ZAxis;
+		vDir = YAxis;
+		break;
+	case CubeFace_NegY:
+		vUp = ZAxis;
+		vDir = -YAxis;
+		break;
+	case CubeFace_PosZ:
+		vDir = ZAxis;
+		break;
+	case CubeFace_NegZ:
+		vDir = -ZAxis;
+		break;
+	}
+	// derive right vector
+	FVector vRight(vUp^ vDir);
+	// create matrix from the 3 axes
+	return FBasisVectorMatrix(vRight, vUp, vDir, FVector::ZeroVector);
+}
+
 UE::MoviePipeline::FImagePassCameraViewData UMoviePipelineImagePassBase::GetCameraInfo(FMoviePipelineRenderPassMetrics& InOutSampleState, IViewCalcPayload* OptPayload) const
 {
 	UE::MoviePipeline::FImagePassCameraViewData OutCameraData;
diff --git a/Engine/Plugins/MovieScene/MovieRenderPipeline/Source/MovieRenderPipelineRenderPasses/Public/MoviePipelineDeferredPasses.h b/Engine/Plugins/MovieScene/MovieRenderPipeline/Source/MovieRenderPipelineRenderPasses/Public/MoviePipelineDeferredPasses.h
--- a/Engine/Plugins/MovieScene/MovieRenderPipeline/Source/MovieRenderPipelineRenderPasses/Public/MoviePipelineDeferredPasses.h
+++ b/Engine/Plugins/MovieScene/MovieRenderPipeline/Source/MovieRenderPipelineRenderPasses/Public/MoviePipelineDeferredPasses.h
@@ -68,6 +68,7 @@ protected:
 	virtual int32 GetOutputFileSortingOrder() const override { return 0; }
 	virtual bool IsAlphaInTonemapperRequiredImpl() const override { return bAccumulatorIncludesAlpha; }
 	virtual FSceneViewStateInterface* GetSceneViewStateInterface(IViewCalcPayload* OptPayload = nullptr) override;
+	virtual FSceneViewStateInterface* GetExposureSceneViewStateInterface(IViewCalcPayload* OptPayload = nullptr, int32 CubeFaceIndex = INDEX_NONE) override;
 	virtual void AddViewExtensions(FSceneViewFamilyContext& InContext, FMoviePipelineRenderPassMetrics& InOutSampleState) override;
 	virtual bool IsAutoExposureAllowed(const FMoviePipelineRenderPassMetrics& InSampleState) const override;
 	virtual void BlendPostProcessSettings(FSceneView* InView, FMoviePipelineRenderPassMetrics& InOutSampleState, IViewCalcPayload* OptPayload = nullptr);
diff --git a/Engine/Plugins/MovieScene/MovieRenderPipeline/Source/MovieRenderPipelineRenderPasses/Public/MoviePipelineImagePassBase.h b/Engine/Plugins/MovieScene/MovieRenderPipeline/Source/MovieRenderPipelineRenderPasses/Public/MoviePipelineImagePassBase.h
--- a/Engine/Plugins/MovieScene/MovieRenderPipeline/Source/MovieRenderPipelineRenderPasses/Public/MoviePipelineImagePassBase.h
+++ b/Engine/Plugins/MovieScene/MovieRenderPipeline/Source/MovieRenderPipelineRenderPasses/Public/MoviePipelineImagePassBase.h
@@ -117,6 +117,8 @@ protected:
 	FVector4 CalculatePrinciplePointOffsetForTiling(const FMoviePipelineRenderPassMetrics& InSampleState) const;
 	virtual void ModifyProjectionMatrixForTiling(const FMoviePipelineRenderPassMetrics& InSampleState, const bool bInOrthographic, FMatrix& InOutProjectionMatrix, float& OutDoFSensorScale) const;
 
+	virtual FMatrix CalcCubeFaceTransform(ECubeFace Face) const final;
+
 protected:
 	virtual void GetViewShowFlags(FEngineShowFlags& OutShowFlag, EViewModeIndex& OutViewModeIndex) const;
 	virtual TSharedPtr<FSceneViewFamilyContext> CalculateViewFamily(FMoviePipelineRenderPassMetrics& InOutSampleState, IViewCalcPayload* OptPayload = nullptr);
@@ -127,6 +129,7 @@ protected:
 	virtual bool IsAntiAliasingSupported() const { return true; }
 	virtual int32 GetOutputFileSortingOrder() const { return -1; }
 	virtual FSceneViewStateInterface* GetSceneViewStateInterface(IViewCalcPayload* OptPayload = nullptr) { return ViewState.GetReference(); }
+	virtual FSceneViewStateInterface* GetExposureSceneViewStateInterface(IViewCalcPayload* OptPayload = nullptr, int32 CubeFaceIndex = INDEX_NONE) { return nullptr; }
 	virtual void AddViewExtensions(FSceneViewFamilyContext& InContext, FMoviePipelineRenderPassMetrics& InOutSampleState) { }
 	virtual bool IsAutoExposureAllowed(const FMoviePipelineRenderPassMetrics& InSampleState) const { return true; }
 	virtual FSceneView* GetSceneViewForSampleState(FSceneViewFamily* ViewFamily, FMoviePipelineRenderPassMetrics& InOutSampleState, IViewCalcPayload* OptPayload = nullptr);
diff --git a/Engine/Source/Runtime/Engine/Private/SceneView.cpp b/Engine/Source/Runtime/Engine/Private/SceneView.cpp
--- a/Engine/Source/Runtime/Engine/Private/SceneView.cpp
+++ b/Engine/Source/Runtime/Engine/Private/SceneView.cpp
@@ -895,6 +895,7 @@ FSceneView::FSceneView(const FSceneViewInitOptions& InitOptions)
 	, bIsViewInfo(false)
 	, bIsSceneCapture(InitOptions.bIsSceneCapture)
 	, bIsSceneCaptureCube(InitOptions.bIsSceneCaptureCube)
+	, bEyeAdaptationAllViewPixels(false)
 	, bSceneCaptureUsesRayTracing(InitOptions.bSceneCaptureUsesRayTracing)
 	, bIsReflectionCapture(InitOptions.bIsReflectionCapture)
 	, bIsPlanarReflection(InitOptions.bIsPlanarReflection)
@@ -1057,6 +1058,12 @@ PRAGMA_ENABLE_DEPRECATION_WARNINGS
 		{
 			EyeAdaptationViewState = Family->Views[0]->State;
 		}
+
+		// User specified exposure view state, for sharing exposure across views (for example, for movie render queue high-res tiled rendering)
+		if (InitOptions.ExposureSceneViewStateInterface)
+		{
+			EyeAdaptationViewState = InitOptions.ExposureSceneViewStateInterface;
+		}
 	}
 
 #if DO_CHECK || USING_CODE_ANALYSIS
diff --git a/Engine/Source/Runtime/Engine/Public/SceneView.h b/Engine/Source/Runtime/Engine/Public/SceneView.h
--- a/Engine/Source/Runtime/Engine/Public/SceneView.h
+++ b/Engine/Source/Runtime/Engine/Public/SceneView.h
@@ -175,6 +175,7 @@ struct FSceneViewInitOptions : public FSceneViewProjectionData
 {
 	const FSceneViewFamily* ViewFamily;
 	FSceneViewStateInterface* SceneViewStateInterface;
+	FSceneViewStateInterface* ExposureSceneViewStateInterface;
 	FSceneViewOwner ViewActor; // TODO: Rename to ViewOwner
 	int32 PlayerIndex;
 	FViewElementDrawer* ViewElementDrawer;
@@ -266,6 +267,7 @@ struct FSceneViewInitOptions : public FSceneViewProjectionData
 	FSceneViewInitOptions()
 		: ViewFamily(NULL)
 		, SceneViewStateInterface(NULL)
+		, ExposureSceneViewStateInterface(NULL)
 		, PlayerIndex(INDEX_NONE)
 		, ViewElementDrawer(NULL)
 		, BackgroundColor(FLinearColor::Transparent)
@@ -1565,6 +1567,9 @@ public:
 	/** Whether the scene capture is a cube map (bIsSceneCapture will also be set). */
 	bool bIsSceneCaptureCube;
 
+	/** Whether to generate histogram eye adaptation from all view pixels, useful with a shared EyeAdaptationViewState.  Pixels for views must form a contiguous rect. */
+	bool bEyeAdaptationAllViewPixels;
+
 	/**
 	 * Whether this view may use ray tracing, for views that are used to render a scene capture.  Use IsRayTracingAllowedForView()
 	 * to test for ray tracing, not this bool.  Filled in based on the bUseRayTracingIfEnabled field in the scene capture component
diff --git a/Engine/Source/Runtime/Renderer/Private/PostProcess/PostProcessing.cpp b/Engine/Source/Runtime/Renderer/Private/PostProcess/PostProcessing.cpp
--- a/Engine/Source/Runtime/Renderer/Private/PostProcess/PostProcessing.cpp
+++ b/Engine/Source/Runtime/Renderer/Private/PostProcess/PostProcessing.cpp
@@ -1196,6 +1196,18 @@ void AddPostProcessingPasses(
 				}
 			}
 
+			// Optionally generate eye adaptation from the entire set of view rects.  Rects must combine to form a contiguous rect!
+			if (View.bEyeAdaptationAllViewPixels && View.Family->Views.Num() > 1)
+			{
+				FIntRect EyeAdaptationRect = View.Family->Views[0]->UnconstrainedViewRect;
+				for (int32 OtherViewIndex = 1; OtherViewIndex < View.Family->Views.Num(); OtherViewIndex++)
+				{
+					EyeAdaptationRect.Union(View.Family->Views[OtherViewIndex]->UnconstrainedViewRect);
+				}
+
+				HistogramSceneColor = FScreenPassTextureSlice::CreateFromScreenPassTexture(GraphBuilder, FScreenPassTexture(OriginalSceneColor.Texture, EyeAdaptationRect));
+			}
+
 			HistogramTexture = AddHistogramPass(
 				GraphBuilder, View,
 				EyeAdaptationParameters,
diff --git a/Engine/Source/Runtime/Renderer/Private/SceneCaptureRendering.cpp b/Engine/Source/Runtime/Renderer/Private/SceneCaptureRendering.cpp
--- a/Engine/Source/Runtime/Renderer/Private/SceneCaptureRendering.cpp
+++ b/Engine/Source/Runtime/Renderer/Private/SceneCaptureRendering.cpp
@@ -741,6 +741,15 @@ TArray<FSceneView*> SetupViewFamilyForSceneCapture(
 
 		FSceneView* View = new FSceneView(ViewInitOptions);
 		
+		// Generate auto-exposure from all cube map faces.  Only affects cube captures with post processing enabled.  Adds 20% to the cost of post process
+		// with the cheapest possible settings (only tonemap and FXAA enabled), or 2% of overall render time in a trival scene (0.09 ms on a high end card
+		// at 1024 size).  If the performance hit was larger, we could consider an opt out CVar, but this seems fine.  Post processing for cube captures
+		// was added in UE5.5, so there wouldn't be a lot of users of the feature affected by this minor perf impact.
+		if (CubemapFaceIndex == CubeFace_MAX && ViewIndex == 0)
+		{
+			View->bEyeAdaptationAllViewPixels = true;
+		}
+
 		if (SceneCaptureComponent)
 		{
 			GetShowOnlyAndHiddenComponents(SceneCaptureComponent, View->HiddenPrimitives, View->ShowOnlyPrimitives);
