MoviePipeline: Updated Quick Render to support orthographic viewports in the "Current Viewport" and "Use Viewport Camera in Sequence" modes. These modes also now support rendering from piloted actors in the viewport.

diff --git a/Engine/Plugins/MovieScene/MovieRenderPipeline/Source/MovieRenderPipelineCore/Private/Graph/Nodes/MovieGraphApplyViewportLookNode.cpp b/Engine/Plugins/MovieScene/MovieRenderPipeline/Source/MovieRenderPipelineCore/Private/Graph/Nodes/MovieGraphApplyViewportLookNode.cpp
--- a/Engine/Plugins/MovieScene/MovieRenderPipeline/Source/MovieRenderPipelineCore/Private/Graph/Nodes/MovieGraphApplyViewportLookNode.cpp
+++ b/Engine/Plugins/MovieScene/MovieRenderPipeline/Source/MovieRenderPipelineCore/Private/Graph/Nodes/MovieGraphApplyViewportLookNode.cpp
@@ -19,14 +19,14 @@ EMovieGraphBranchRestriction UMovieGraphApplyViewportLookNode::GetBranchRestrict
 }
 
 #if WITH_EDITOR
-FEditorViewportClient* UMovieGraphApplyViewportLookNode::GetViewportClient()
+FLevelEditorViewportClient* UMovieGraphApplyViewportLookNode::GetViewportClient()
 {
-	// Get the currently active *perspective* viewport that's in use. Only perspective viewports are currently supported in Quick Render.
+	// Get the currently active viewport that's in use (both perspective and orthographic viewports are supported)
 	if (const FViewport* ActiveViewport = GEditor->GetActiveViewport())
 	{
-		for (FEditorViewportClient* LevelViewportClient : GEditor->GetLevelViewportClients())
+		for (FLevelEditorViewportClient* LevelViewportClient : GEditor->GetLevelViewportClients())
 		{
-			if (LevelViewportClient->IsPerspective() && (LevelViewportClient->Viewport == ActiveViewport))
+			if (LevelViewportClient->Viewport == ActiveViewport)
 			{
 				return LevelViewportClient;
 			}
diff --git a/Engine/Plugins/MovieScene/MovieRenderPipeline/Source/MovieRenderPipelineCore/Public/Graph/Nodes/MovieGraphApplyViewportLookNode.h b/Engine/Plugins/MovieScene/MovieRenderPipeline/Source/MovieRenderPipelineCore/Public/Graph/Nodes/MovieGraphApplyViewportLookNode.h
--- a/Engine/Plugins/MovieScene/MovieRenderPipeline/Source/MovieRenderPipelineCore/Public/Graph/Nodes/MovieGraphApplyViewportLookNode.h
+++ b/Engine/Plugins/MovieScene/MovieRenderPipeline/Source/MovieRenderPipelineCore/Public/Graph/Nodes/MovieGraphApplyViewportLookNode.h
@@ -7,7 +7,7 @@
 
 #include "MovieGraphApplyViewportLookNode.generated.h"
 
-class FEditorViewportClient;
+class FLevelEditorViewportClient;
 class FSceneView;
 
 // Note: This node is marked as HideDropdown so it does not appear in the node creation menu within the graph. It is meant to
@@ -28,7 +28,7 @@ public:
 
 #if WITH_EDITOR
 	/** Gets the viewport client for the currently active viewport. */
-	static FEditorViewportClient* GetViewportClient();
+	static FLevelEditorViewportClient* GetViewportClient();
 	
 	/** Gets the show flags and view mode index for the currently active viewport. */
 	bool GetViewportInfo(FEngineShowFlags& OutShowFlags, EViewModeIndex& OutViewModeIndex) const;
diff --git a/Engine/Plugins/MovieScene/MovieRenderPipeline/Source/MovieRenderPipelineEditor/Private/Graph/MovieGraphQuickRender.cpp b/Engine/Plugins/MovieScene/MovieRenderPipeline/Source/MovieRenderPipelineEditor/Private/Graph/MovieGraphQuickRender.cpp
--- a/Engine/Plugins/MovieScene/MovieRenderPipeline/Source/MovieRenderPipelineEditor/Private/Graph/MovieGraphQuickRender.cpp
+++ b/Engine/Plugins/MovieScene/MovieRenderPipeline/Source/MovieRenderPipelineEditor/Private/Graph/MovieGraphQuickRender.cpp
@@ -33,6 +33,7 @@
 
 #if WITH_EDITOR
 #include "EditorViewportClient.h"
+#include "LevelEditorViewport.h"
 #include "Selection.h"
 #endif
 
@@ -719,10 +720,88 @@ void FMovieGraphQuickRender::PerformPreRenderSetup(const EMovieGraphQuickRenderM
 		ACameraActor* NewCamera = Cast<ACameraActor>(PieWorld->SpawnActor<ACameraActor>(CameraSpawnParams));
 
 		// Update the camera to look like the viewport's camera
-		const FEditorViewportClient* ViewportClient = UMovieGraphApplyViewportLookNode::GetViewportClient();
-		NewCamera->SetActorLocation(ViewportClient->GetViewLocation(), false);
-		NewCamera->SetActorRotation(ViewportClient->GetViewRotation());
-		NewCamera->GetCameraComponent()->SetFieldOfView(ViewportClient->ViewFOV);
+		if (const FLevelEditorViewportClient* ViewportClient = UMovieGraphApplyViewportLookNode::GetViewportClient())
+		{
+			// Initialize the camera properties to reflect the (non-pilot) view shown in the viewport 
+			FVector CameraPosition = ViewportClient->GetViewLocation();
+			FRotator CameraRotation = ViewportClient->GetViewRotation();
+			bool bIsPerspective = ViewportClient->IsPerspective();
+			float CameraFov = ViewportClient->ViewFOV;
+			float CameraAspectRatio = ViewportClient->AspectRatio;
+			float CameraOrthoWidth = ViewportClient->Viewport->GetSizeXY().X * ViewportClient->GetOrthoUnitsPerPixel(ViewportClient->Viewport);
+			
+			// If the viewport is being piloted by a camera, then we should use the properties of that to render from instead of the viewport's
+			// non-piloted camera properties.
+			const FLevelViewportActorLock& ActorLock = ViewportClient->GetActorLock();
+			if (ActorLock.HasValidLockedActor())
+			{
+				// Support basic position/rotation for non-camera actors
+				const TStrongObjectPtr<AActor> LockedActorPin = ActorLock.LockedActor.Pin();
+				CameraPosition = LockedActorPin->GetActorLocation();
+				CameraRotation = LockedActorPin->GetActorRotation();
+
+				// If the piloted actor is a camera, get all of its camera properties. Note that this supports "cinematic locks" (like cameras piloted
+				// via Sequencer) and normal piloted cameras.
+				if (const UCameraComponent* PilotCameraComponent = ViewportClient->GetCameraComponentForView())
+				{
+					CameraPosition = PilotCameraComponent->GetTypedOuter<AActor>()->GetActorLocation();
+					CameraRotation = PilotCameraComponent->GetTypedOuter<AActor>()->GetActorRotation();
+					bIsPerspective = PilotCameraComponent->ProjectionMode == ECameraProjectionMode::Perspective;
+					CameraFov = PilotCameraComponent->FieldOfView;
+					CameraAspectRatio = PilotCameraComponent->AspectRatio;
+					CameraOrthoWidth = PilotCameraComponent->OrthoWidth;
+				}
+			}
+			else
+			{
+				// The camera properties were initialized from the viewport camera properties, so all we need to do here is take into account the
+				// special case of the viewport showing an axis-aligned ortho cam. Getting the actor rotation from the viewport client for these
+				// ortho cams does not work properly, so their rotation is manually specified here. Sharing this logic with the viewport would involve
+				// some fairly significant changes, so for now MRG calculates this separately.
+				if (!bIsPerspective)
+				{
+					switch (ViewportClient->ViewportType)
+					{
+					case LVT_OrthoTop:
+						CameraRotation = FRotator(-90.0f, -180.0f, 0.0f);
+						break;
+					case LVT_OrthoBottom:
+						CameraRotation = FRotator(90.0f, 0.0f, 0.0f);
+						break;
+					case LVT_OrthoLeft:
+						CameraRotation = FRotator(0.0f, -90.0f, 0.0f);
+						break;
+					case LVT_OrthoRight:
+						CameraRotation = FRotator(0.0f, 90.0f, 0.0f);
+						break;
+					case LVT_OrthoBack:
+						CameraRotation = FRotator(0.0f, 0.0f, 0.0f);
+						break;
+					case LVT_OrthoFront:
+						CameraRotation = FRotator(0.0f, 180.0f, 0.0f);
+						break;
+					default:
+						CameraRotation = FRotator(0.0f, 0.0f, 0.0f);
+					}
+				}
+			}
+
+			NewCamera->SetActorLocation(CameraPosition);
+			NewCamera->SetActorRotation(CameraRotation);
+
+			// Note: At extreme zoom levels, the ortho clipping planes don't match up exactly with the viewport. This is something to improve in the future.
+			UCameraComponent* CameraComponent = NewCamera->GetCameraComponent();
+			CameraComponent->SetProjectionMode(bIsPerspective ? ECameraProjectionMode::Perspective : ECameraProjectionMode::Orthographic);
+			CameraComponent->SetAspectRatio(CameraAspectRatio);
+			CameraComponent->SetFieldOfView(CameraFov);
+			CameraComponent->SetOrthoWidth(CameraOrthoWidth);
+			CameraComponent->SetAutoCalculateOrthoPlanes(true);
+		}
+		else
+		{
+			// This case should be exceedingly rare, but log about it anyway
+			UE_LOG(LogMovieRenderPipeline, Error, TEXT("Quick Render: Could not determine the active viewport to render from, so the camera used will be incorrect."));
+		}
 
 		return NewCamera;
 	};
diff --git a/Engine/Plugins/MovieScene/MovieRenderPipeline/Source/MovieRenderPipelineRenderPasses/Private/Graph/Renderers/MovieGraphImagePassBase.cpp b/Engine/Plugins/MovieScene/MovieRenderPipeline/Source/MovieRenderPipelineRenderPasses/Private/Graph/Renderers/MovieGraphImagePassBase.cpp
--- a/Engine/Plugins/MovieScene/MovieRenderPipeline/Source/MovieRenderPipelineRenderPasses/Private/Graph/Renderers/MovieGraphImagePassBase.cpp
+++ b/Engine/Plugins/MovieScene/MovieRenderPipeline/Source/MovieRenderPipelineRenderPasses/Private/Graph/Renderers/MovieGraphImagePassBase.cpp
@@ -21,6 +21,7 @@
 
 #if WITH_EDITOR
 #include "EditorViewportClient.h"
+#include "LevelEditorViewport.h"
 #endif
 
 
@@ -184,7 +185,7 @@ TSharedRef<FSceneViewFamilyContext> FMovieGraphImagePassBase::CreateSceneViewFam
 				}
 			}
 
-			if (const FEditorViewportClient* ViewportClient = UMovieGraphApplyViewportLookNode::GetViewportClient())
+			if (const FLevelEditorViewportClient* ViewportClient = UMovieGraphApplyViewportLookNode::GetViewportClient())
 			{
 				ExposureSettings = ViewportClient->ExposureSettings;
 			}
