Build fix for unresolved external ComputeAutoExposureCubeCaptureSize.

#rnx

diff --git a/Engine/Plugins/MovieScene/MoviePipelineMaskRenderPass/Source/MoviePipelineMaskRenderPass/Private/MoviePipelinePanoramicPass.cpp b/Engine/Plugins/MovieScene/MoviePipelineMaskRenderPass/Source/MoviePipelineMaskRenderPass/Private/MoviePipelinePanoramicPass.cpp
--- a/Engine/Plugins/MovieScene/MoviePipelineMaskRenderPass/Source/MoviePipelineMaskRenderPass/Private/MoviePipelinePanoramicPass.cpp
+++ b/Engine/Plugins/MovieScene/MoviePipelineMaskRenderPass/Source/MoviePipelineMaskRenderPass/Private/MoviePipelinePanoramicPass.cpp
@@ -92,7 +92,7 @@ namespace MoviePipeline
 }
 
 // Cube capture is arranged in 3x2 square tiles, rounded down to a multiple of 8 pixels.
-int32 ComputeAutoExposureCubeCaptureSize(FIntPoint Resolution)
+static int32 ComputeAutoExposureCubeCaptureSize(FIntPoint Resolution)
 {
 	return AlignDown(FMath::Min(Resolution.X / 3, Resolution.Y / 2), 8);
 }
diff --git a/Engine/Plugins/MovieScene/MovieRenderPipeline/Source/MovieRenderPipelineRenderPasses/Private/MoviePipelineImagePassBase.cpp b/Engine/Plugins/MovieScene/MovieRenderPipeline/Source/MovieRenderPipelineRenderPasses/Private/MoviePipelineImagePassBase.cpp
--- a/Engine/Plugins/MovieScene/MovieRenderPipeline/Source/MovieRenderPipelineRenderPasses/Private/MoviePipelineImagePassBase.cpp
+++ b/Engine/Plugins/MovieScene/MovieRenderPipeline/Source/MovieRenderPipelineRenderPasses/Private/MoviePipelineImagePassBase.cpp
@@ -456,7 +456,11 @@ void UMoviePipelineImagePassBase::GatherOutputPassesImpl(TArray<FMoviePipelinePa
 	ExpectedRenderPasses.Add(PassIdentifier);
 }
 
-extern int32 ComputeAutoExposureCubeCaptureSize(FIntPoint Resolution);
+// Cube capture is arranged in 3x2 square tiles, rounded down to a multiple of 8 pixels.
+static int32 ComputeAutoExposureCubeCaptureSize(FIntPoint Resolution)
+{
+	return AlignDown(FMath::Min(Resolution.X / 3, Resolution.Y / 2), 8);
+}
 
 FSceneView* UMoviePipelineImagePassBase::GetSceneViewForSampleState(FSceneViewFamily* ViewFamily, FMoviePipelineRenderPassMetrics& InOutSampleState, IViewCalcPayload* OptPayload)
 {
