MovieRenderPipeline: Added basic support for multi-part EXR files, which allow multiple headers with different data window sizes, which allow burn-ins rendered at output resolution to be included in the same EXR as the overscanned render instead of as a separate file when using multi-layering. Also changed EXR data and display window computations so that display window top left is always (0, 0) and data window can be negative for overscan.

diff --git a/Engine/Plugins/MovieScene/MovieRenderPipeline/Source/MovieRenderPipelineCore/Public/MovieRenderPipelineDataTypes.h b/Engine/Plugins/MovieScene/MovieRenderPipeline/Source/MovieRenderPipelineCore/Public/MovieRenderPipelineDataTypes.h
--- a/Engine/Plugins/MovieScene/MovieRenderPipeline/Source/MovieRenderPipelineCore/Public/MovieRenderPipelineDataTypes.h
+++ b/Engine/Plugins/MovieScene/MovieRenderPipeline/Source/MovieRenderPipelineCore/Public/MovieRenderPipelineDataTypes.h
@@ -1730,17 +1730,28 @@ struct TAsyncCropImage
 	
 	void operator()(FImagePixelData* PixelData)
 	{
-		const bool bIsSrcBigEnough = PixelData->GetSize().X >= CropRect.Size().X && PixelData->GetSize().Y >= CropRect.Size().Y;
-		if (!ensureMsgf(bIsSrcBigEnough, TEXT("Cannot crop image of size (%d,%d) using crop rectangle of size (%d,%d)"),
-			PixelData->GetSize().X, PixelData->GetSize().Y, CropRect.Size().X, CropRect.Size().Y))
+		// Perhaps a little unintuitively, crop region is allowed to be bigger than the original source image, in which case the output image will be scaled up
+		// to the size of the crop region and filled with zeros, while the source image is copied into the subregion it inhabits in the crop region.
+		// Still have to make sure there is some overlap between the crop region and the source image (it's not useful to return an entirely blank output)
+		const bool bIsValidCropRegion = CropRect.Max.X > 0 && CropRect.Max.Y > 0 && CropRect.Min.X < PixelData->GetSize().X && CropRect.Min.Y < PixelData->GetSize().Y;
+		if (!ensureMsgf(bIsValidCropRegion, TEXT("Crop rectangle (%d, %d) to (%d, %d) does not overlap any part of the source image of size (%d, %d)"),
+			CropRect.Min.X, CropRect.Min.Y, CropRect.Max.X, CropRect.Max.Y, PixelData->GetSize().X, PixelData->GetSize().Y))
 		{
 			return;
 		}
 
 		TRACE_CPUPROFILER_EVENT_SCOPE(MoviePipeline::TAsyncCropImage)
 		
+		// In the case where the crop region extends beyond the top-left corner of the source image, keep track of that offset so that
+		// we can read and write valid pixel locations from the source and destination images (e.g. normal crop only reads from the crop rect, but
+		// inflationary crop can only read from the top left of the source image even if the crop rect extends beyond it). SrcOffset is the top left
+		// of the rectangle we read from in the source image, and DestOffset is the top left of the rectangle we write to in the destination image
+		const FIntPoint SrcOffset = FIntPoint(FMath::Max(CropRect.Min.X, 0), FMath::Max(CropRect.Min.Y, 0));
+		const FIntPoint DestOffset = FIntPoint(FMath::Max(-CropRect.Min.X, 0), FMath::Max(-CropRect.Min.Y, 0));
+
+		// Allocate enough pixels for the crop region and fill them with zero values
 		TArray64<PixelType> DefaultPixels;
-		DefaultPixels.AddUninitialized(CropRect.Size().X * CropRect.Size().Y);
+		DefaultPixels.AddZeroed(CropRect.Size().X * CropRect.Size().Y);
 
 		TSharedRef<FImagePixelDataPayload> NewPayload = PixelData->GetPayload<FImagePixelDataPayload>()->Copy();
 		
@@ -1748,13 +1759,20 @@ struct TAsyncCropImage
 		TImagePixelData<PixelType>* DestColorData = static_cast<TImagePixelData<PixelType>*>(CroppedImage.Get());
 		TImagePixelData<PixelType>* SrcColorData = static_cast<TImagePixelData<PixelType>*>(PixelData);
 
-		ParallelFor(DestColorData->GetSize().Y,
+		// Find the size of the region that actually needs to be copied, which is the intersection of the source image and the crop rectangle.
+		// Start reading from the top left of the crop region, or if it extends beyond the source image, the source image top left, and continue
+		// to the bottom right corner of the crop rect or the source image, whichever is smaller
+		const FIntPoint Size = FIntPoint(
+			FMath::Min(CropRect.Max.X - SrcOffset.X, SrcColorData->GetSize().X),
+			FMath::Min(CropRect.Max.Y - SrcOffset.Y, SrcColorData->GetSize().Y));
+		
+		ParallelFor(Size.Y,
 			[&](int32 ScanlineIndex = 0)
 			{
-				int64 DstIndex = int64(ScanlineIndex) * int64(DestColorData->GetSize().X);
-				int64 SrcIndex = (int64(ScanlineIndex) + int64(CropRect.Min.Y)) * int64(SrcColorData->GetSize().X) + int64(CropRect.Min.X);
+				int64 DstIndex = (int64(ScanlineIndex) + int64(DestOffset.Y)) * int64(DestColorData->GetSize().X) + int64(DestOffset.X);
+				int64 SrcIndex = (int64(ScanlineIndex) + int64(SrcOffset.Y)) * int64(SrcColorData->GetSize().X) + int64(SrcOffset.X);
 				
-				FMemory::Memcpy(&DestColorData->Pixels[DstIndex], &SrcColorData->Pixels[SrcIndex], DestColorData->GetSize().X * sizeof(PixelType));
+				FMemory::Memcpy(&DestColorData->Pixels[DstIndex], &SrcColorData->Pixels[SrcIndex], Size.X * sizeof(PixelType));
 			});
 
 		// If we have a parent write task that is holding the image being processed, write to that. Otherwise,
diff --git a/Engine/Plugins/MovieScene/MovieRenderPipeline/Source/MovieRenderPipelineRenderPasses/Private/MovieGraphImageSequenceOutputNode.cpp b/Engine/Plugins/MovieScene/MovieRenderPipeline/Source/MovieRenderPipelineRenderPasses/Private/MovieGraphImageSequenceOutputNode.cpp
--- a/Engine/Plugins/MovieScene/MovieRenderPipeline/Source/MovieRenderPipelineRenderPasses/Private/MovieGraphImageSequenceOutputNode.cpp
+++ b/Engine/Plugins/MovieScene/MovieRenderPipeline/Source/MovieRenderPipelineRenderPasses/Private/MovieGraphImageSequenceOutputNode.cpp
@@ -281,7 +281,7 @@ void UMovieGraphImageSequenceOutputNode::OnReceiveImageDataImpl(UMovieGraphPipel
 }
 
 
-TUniquePtr<FEXRImageWriteTask> UMovieGraphImageSequenceOutputNode_EXR::CreateImageWriteTask(FString InFileName, EEXRCompressionFormat InCompression) const
+TUniquePtr<FEXRImageWriteTask> UMovieGraphImageSequenceOutputNode_EXR::CreateImageWriteTask(FString InFileName, EEXRCompressionFormat InCompression, bool bMultiPart) const
 {
 	// Ensure our OpenExrRTTI module gets loaded.
 	UE_CALL_ONCE([]
@@ -290,8 +290,14 @@ TUniquePtr<FEXRImageWriteTask> UMovieGraphImageSequenceOutputNode_EXR::CreateIma
 			FModuleManager::Get().LoadModule(TEXT("UEOpenExrRTTI"));
 		});
 
+	// If not using multi-part, we have to pad all layers up to the maximum resolution. If multi-part is on, different header
+	// data window sizes are suppported, so check the cvar to see if we should pad
+	const bool bPadToDataWindowSize = !bMultiPart || UE::MoviePipeline::CVarMoviePipelinePadLayersForMultiPartEXR.GetValueOnGameThread();
+	
 	TUniquePtr<FEXRImageWriteTask> ImageWriteTask = MakeUnique<FEXRImageWriteTask>();
 	ImageWriteTask->Filename = MoveTemp(InFileName);
+	ImageWriteTask->bMultipart = bMultiPart;
+	ImageWriteTask->bPadToDataWindowSize = bPadToDataWindowSize;
 	ImageWriteTask->Compression = InCompression;
 	// ImageWriteTask->CompressionLevel is intentionally skipped and not exposed ("dwaCompressionLevel" is deprecated)
 
@@ -493,22 +499,19 @@ void UMovieGraphImageSequenceOutputNode_MultiLayerEXR::OnReceiveImageDataImpl(UM
 		UMovieGraphNode::GlobalsPinName, bIncludeCDOs, bExactMatch);
 	checkf(ParentNode, TEXT("Multi-Layer EXR should not exist without a parent node in the graph."));
 
-	// Generate a mapping of resolved filename -> RenderIDs, filename -> resolve args, and filename -> frame templated filename. The generated EXRs can only
-	// store layers with a common resolution, and this takes care of ensuring that only renderIDs with a common resolution
-	// map to the same filename.
-	TMap<FString, TArray<FMovieGraphRenderDataIdentifier>> FilenameToRenderIDs;
-	TMap<FString, FMovieGraphResolveArgs> FilenameToResolveArgs;
-	TMap<FString, FString> FilenameToFrameTemplatedFilename;
-	GetFilenameToRenderIDMappings(ParentNode, InPipeline, InRawFrameData, FilenameToRenderIDs, FilenameToResolveArgs, FilenameToFrameTemplatedFilename);
+	// Generate the output config for each filename, which contains the Render IDs, the resolve args, the frame-templated filename, and the maximum resolution
+	// to use when outputting to the corresponding EXR file
+	TMap<FString, FEXROutputConfigForFilename> FilenameToRenderConfig;
+	GetFilenameToEXROutputConfigMappings(ParentNode, InPipeline, InRawFrameData, FilenameToRenderConfig);
 
 	// Write an EXR for each filename, which potentially contains multiple passes (render IDs).
-	for (const TPair<FString, TArray<FMovieGraphRenderDataIdentifier>>& RenderIDsForFilename : FilenameToRenderIDs)
+	for (TPair<FString, FEXROutputConfigForFilename>& RenderConfigForFilename : FilenameToRenderConfig)
 	{
-		const FString& Filename = RenderIDsForFilename.Key;
-		const TArray<FMovieGraphRenderDataIdentifier>& RenderIDs = RenderIDsForFilename.Value;
+		const FString& Filename = RenderConfigForFilename.Key;
+		FEXROutputConfigForFilename& RenderConfig = RenderConfigForFilename.Value;
 		
-		TUniquePtr<FEXRImageWriteTask> MultiLayerImageTask = CreateImageWriteTask(Filename, ParentNode->Compression);
-		PrepareTaskGlobalMetadata(*MultiLayerImageTask, InRawFrameData, FilenameToResolveArgs[Filename].FileMetadata);
+		TUniquePtr<FEXRImageWriteTask> MultiLayerImageTask = CreateImageWriteTask(Filename, ParentNode->Compression, ParentNode->bMultipart);
+		PrepareTaskGlobalMetadata(*MultiLayerImageTask, InRawFrameData, RenderConfig.ResolveArgs.FileMetadata);
 
 		// Keep track of the lowest render layer index found among the layers that are included. This will be used as the index provided to the output
 		// future. This index is used to determine what the first render layer is when "First Render Layer Only" is turned on for displaying media
@@ -519,7 +522,7 @@ void UMovieGraphImageSequenceOutputNode_MultiLayerEXR::OnReceiveImageDataImpl(UM
 		bool bHasGeneratedPrimaryRGBALayer = false;
 		int32 LayerIndex = 0;
 		int32 ShotIndex = 0;
-		for (const FMovieGraphRenderDataIdentifier& RenderID : RenderIDs)
+		for (const FMovieGraphRenderDataIdentifier& RenderID : RenderConfig.RenderIDs)
 		{
 			const TUniquePtr<FImagePixelData>& ImageData = InRawFrameData->ImageOutputData[RenderID];
 			checkf(ImageData.IsValid(), TEXT("Unexpected empty image data: incorrectly moved or its production failed?"));
@@ -610,19 +613,17 @@ void UMovieGraphImageSequenceOutputNode_MultiLayerEXR::OnReceiveImageDataImpl(UM
 #endif // WITH_OCIO
 			UpdateTaskPerLayer(*MultiLayerImageTask, ParentNode, MoveTemp(PixelData), LayerIndex, LayerName, ResolvedOCIOContext);
 
+			// Ensure the write task uses the maximum resolution of all the layers being written
+			MultiLayerImageTask->Width = RenderConfig.MaximumResolution.X;
+			MultiLayerImageTask->Height = RenderConfig.MaximumResolution.Y;
+			
 			LayerIndex++;
 		}
-		
-		FString FrameTemplatedFilename;
-		if (FString* FrameTemplatedFilenamePtr = FilenameToFrameTemplatedFilename.Find(Filename))
-		{
-			FrameTemplatedFilename = *FrameTemplatedFilenamePtr;
-		}
-		
+
 		UE::MovieGraph::FMovieGraphOutputFutureData OutputFutureData;
 		OutputFutureData.Shot = InPipeline->GetActiveShotList()[ShotIndex];
 		OutputFutureData.FilePath = Filename;
-		OutputFutureData.FrameTemplatedFilePath = FrameTemplatedFilename;
+		OutputFutureData.FrameTemplatedFilePath = RenderConfig.FrameTemplatedFilename;
 		OutputFutureData.DataIdentifier = FMovieGraphRenderDataIdentifier(); // EXRs put all the render passes internally so this resolves to a ""
 		OutputFutureData.OriginNodeClass = GetClass();
 		OutputFutureData.RenderLayerIndex = LowestRenderLayerIndex;
@@ -631,12 +632,9 @@ void UMovieGraphImageSequenceOutputNode_MultiLayerEXR::OnReceiveImageDataImpl(UM
 	}
 }
 
-void UMovieGraphImageSequenceOutputNode_MultiLayerEXR::GetFilenameToRenderIDMappings(
-	const UMovieGraphImageSequenceOutputNode_MultiLayerEXR* InParentNode,
-	UMovieGraphPipeline* InPipeline, UE::MovieGraph::FMovieGraphOutputMergerFrame* InRawFrameData,
-	TMap<FString, TArray<FMovieGraphRenderDataIdentifier>>& OutFilenameToRenderIDs,
-	TMap<FString, FMovieGraphResolveArgs>& OutFilenameToResolveArgs,
-	TMap<FString, FString>& OutFilenameToFrameTemplatedFilename) const
+void UMovieGraphImageSequenceOutputNode_MultiLayerEXR::GetFilenameToEXROutputConfigMappings(
+	const UMovieGraphImageSequenceOutputNode_MultiLayerEXR* InParentNode, UMovieGraphPipeline* InPipeline,
+	UE::MovieGraph::FMovieGraphOutputMergerFrame* InRawFrameData, TMap<FString, FEXROutputConfigForFilename>& OutFilenameToOutputConfigs) const
 {
 	// Merge one layer's resolve args (InNewResolveArgs) into an existing set of resolve args (InExistingResolveArgs).
 	auto MergeResolveArgs = [](FMovieGraphResolveArgs& InNewResolveArgs, FMovieGraphResolveArgs& InExistingResolveArgs)
@@ -664,8 +662,6 @@ void UMovieGraphImageSequenceOutputNode_MultiLayerEXR::GetFilenameToRenderIDMapp
 		InExistingResolveArgs.FilenameArguments = InNewResolveArgs.FilenameArguments;
 	};
 	
-	TMap<FString, TArray<FIntPoint>> FilenameToResolutions;
-
 	// First, generate filename -> renderID mapping, and filename -> resolution mapping.
 	// This assumes that all render passes will have the same resolution, so we use 0 as the resolution index.
 	// Once we know the resolutions of all the render passes, they can be binned together into groups with the same
@@ -683,60 +679,23 @@ void UMovieGraphImageSequenceOutputNode_MultiLayerEXR::GetFilenameToRenderIDMapp
 		FString FrameTemplatedFilename;
 		constexpr int32 ResolutionIndex = 0;
 		FMovieGraphResolveArgs ResolveArgs;
-		const FString PreliminaryFileName = ResolveOutputFilename(InParentNode, InPipeline, ResolutionIndex, InRawFrameData, RenderPassData.Key, ResolveArgs, FrameTemplatedFilename);
-		
-		TArray<FMovieGraphRenderDataIdentifier>& RenderIDs = OutFilenameToRenderIDs.FindOrAdd(PreliminaryFileName);
-		RenderIDs.Add(RenderPassData.Key);
-
-		OutFilenameToFrameTemplatedFilename.Add(PreliminaryFileName, FrameTemplatedFilename);
+		const FString PreliminaryFileName = ResolveOutputFilename(InParentNode, InPipeline, InRawFrameData, RenderPassData.Key, ResolveArgs, FrameTemplatedFilename);
 
-		TArray<FIntPoint>& Resolutions = FilenameToResolutions.FindOrAdd(PreliminaryFileName);
-		Resolutions.AddUnique(RenderPassData.Value->GetSize());
-
-		MergeResolveArgs(ResolveArgs, OutFilenameToResolveArgs.FindOrAdd(PreliminaryFileName));
-	}
-
-	// Second, re-generate filenames if any render passes of differing resolutions map to the same file.
-	for (const TPair<FString, TArray<FIntPoint>>& FilenameAndResolutions : FilenameToResolutions)
-	{
-		const FString& PreliminaryFilename = FilenameAndResolutions.Key;
-		const TArray<FIntPoint>& Resolutions = FilenameAndResolutions.Value;
-
-		// If there's only one resolution for this filename, there's nothing to disambiguate
-		if (Resolutions.Num() == 1)
-		{
-			continue;
-		}
-
-		// If there IS more than one resolution for this filename, then disambiguate
-		const TArray<FMovieGraphRenderDataIdentifier> OldRenderIDs = OutFilenameToRenderIDs.FindAndRemoveChecked(PreliminaryFilename);
-		OutFilenameToResolveArgs.Remove(PreliminaryFilename);
-		for (const FMovieGraphRenderDataIdentifier& RenderID : OldRenderIDs)
-		{
-			TUniquePtr<FImagePixelData>& ImageData = InRawFrameData->ImageOutputData[RenderID];
-			const int32 ResolutionIndex = Resolutions.IndexOfByKey(ImageData->GetSize());
-
-			// Re-resolve the filename, this time using the resolution index to generate a filename that will only contain
-			// passes with this particular resolution
-			FString FrameTemplatedFilename;
-			FMovieGraphResolveArgs ResolveArgs;
-			const FString FinalFilename = ResolveOutputFilename(InParentNode, InPipeline, ResolutionIndex, InRawFrameData, RenderID, ResolveArgs, FrameTemplatedFilename);
-
-			TArray<FMovieGraphRenderDataIdentifier>& RenderIDs = OutFilenameToRenderIDs.FindOrAdd(FinalFilename);
-			RenderIDs.Add(RenderID);
-
-			OutFilenameToFrameTemplatedFilename.Remove(PreliminaryFilename);
-			OutFilenameToFrameTemplatedFilename.Add(FinalFilename, FrameTemplatedFilename);
+		FEXROutputConfigForFilename& OutputConfig = OutFilenameToOutputConfigs.FindOrAdd(PreliminaryFileName);
+		
+		OutputConfig.RenderIDs.Add(RenderPassData.Key);
+		OutputConfig.FrameTemplatedFilename = FrameTemplatedFilename;
+		
+		OutputConfig.MaximumResolution.X = FMath::Max(OutputConfig.MaximumResolution.X, RenderPassData.Value->GetSize().X);
+		OutputConfig.MaximumResolution.Y = FMath::Max(OutputConfig.MaximumResolution.Y, RenderPassData.Value->GetSize().Y);
 
-			MergeResolveArgs(ResolveArgs, OutFilenameToResolveArgs.FindOrAdd(FinalFilename));
-		}
+		MergeResolveArgs(ResolveArgs, OutputConfig.ResolveArgs);
 	}
 }
 
 FString UMovieGraphImageSequenceOutputNode_MultiLayerEXR::ResolveOutputFilename(
 	const UMovieGraphImageSequenceOutputNode_MultiLayerEXR* InParentNode,
 	const UMovieGraphPipeline* InPipeline,
-	const int32 ResolutionIndex,
 	const UE::MovieGraph::FMovieGraphOutputMergerFrame* InRawFrameData,
 	const FMovieGraphRenderDataIdentifier& InRenderDataIdentifier,
 	FMovieGraphResolveArgs& OutResolveArgs,
@@ -751,8 +710,7 @@ FString UMovieGraphImageSequenceOutputNode_MultiLayerEXR::ResolveOutputFilename(
 		return FString();
 	}
 	
-	// If we have more than one resolution we'll store it as "_Add" / "_Add(1)" etc via {ExtraTag}.
-	FString FileNameFormatString = InParentNode->FileNameFormat + "{ExtraTag}";
+	FString FileNameFormatString = InParentNode->FileNameFormat;
 	
 	// If we're writing more than one render pass out, we need to ensure the file name has the format string in it so we don't
 	// overwrite the same file multiple times. Burn In overlays don't count because they get composited on top of an existing file.
@@ -766,22 +724,6 @@ FString UMovieGraphImageSequenceOutputNode_MultiLayerEXR::ResolveOutputFilename(
 	FormatOverrides.Add(TEXT("render_pass"), TEXT("")); // Render Passes are included inside the exr file by named layers.
 	FormatOverrides.Add(TEXT("ext"), Extension);
 
-	// The logic for the ExtraTag is a little complicated. If there's only one layer (ideal situation) then it's empty.
-	if (ResolutionIndex == 0)
-	{
-		FormatOverrides.Add(TEXT("ExtraTag"), TEXT(""));
-	}
-	else if (ResolutionIndex == 1)
-	{
-		// This is our most common case when we have a second file (the only expected one really)
-		FormatOverrides.Add(TEXT("ExtraTag"), TEXT("_Add"));
-	}
-	else
-	{
-		// Finally a fallback in the event we have three or more unique resolutions.
-		FormatOverrides.Add(TEXT("ExtraTag"), FString::Printf(TEXT("_Add(%d)"), ResolutionIndex));
-	}
-
 	// The layer's render data identifier is used here in the resolve. Usually this is not a problem. However, the user may include some tokens, like
 	// {layer_name}, that come from the identifier, which will prevent all layers from being placed in the same multi-layer EXR (because now the path
 	// isn't resolving to the path that other layers are resolving to). We have to assume that the user is doing this intentionally, even though it's
diff --git a/Engine/Plugins/MovieScene/MovieRenderPipeline/Source/MovieRenderPipelineRenderPasses/Private/MoviePipelineEXROutput.cpp b/Engine/Plugins/MovieScene/MovieRenderPipeline/Source/MovieRenderPipelineRenderPasses/Private/MoviePipelineEXROutput.cpp
--- a/Engine/Plugins/MovieScene/MovieRenderPipeline/Source/MovieRenderPipelineRenderPasses/Private/MoviePipelineEXROutput.cpp
+++ b/Engine/Plugins/MovieScene/MovieRenderPipeline/Source/MovieRenderPipelineRenderPasses/Private/MoviePipelineEXROutput.cpp
@@ -20,19 +20,36 @@
 #include "MoviePipelineUtils.h"
 #include "ColorManagement/ColorSpace.h"
 #include "HDRHelper.h"
+#include "Algo/Accumulate.h"
 
 THIRD_PARTY_INCLUDES_START
 #include "OpenEXR/ImfChannelList.h"
 #include "OpenEXR/ImfStandardAttributes.h"
-
-#include UE_INLINE_GENERATED_CPP_BY_NAME(MoviePipelineEXROutput)
+#include "OpenEXR/ImfMultiPartOutputFile.h"
+#include "OpenEXR/ImfOutputPart.h"
 
 THIRD_PARTY_INCLUDES_END
 
+#include UE_INLINE_GENERATED_CPP_BY_NAME(MoviePipelineEXROutput)
+
 #if WITH_UNREALEXR
 
 // NOTE: see also ExrImageWrapper
 
+namespace UE
+{
+	namespace MoviePipeline
+	{
+		TAutoConsoleVariable<bool> CVarMoviePipelinePadLayersForMultiPartEXR(
+			TEXT("MoviePipeline.PadLayersForMultiPartEXR"),
+			true,
+			TEXT("Indicates that layers in a multi-part EXR file should be padded to match the resolution of the largest layer.\n")
+			TEXT("When enabled, padding will be applied to all layers to ensure their data windows match when written to a multi-part EXR file. When disabled,\n")
+			TEXT("each layer will have its own resolution and data window, which may reduce filesize, but not all software supports this when using multi-part EXR files"),
+			ECVF_Default);
+	}
+}
+
 class FExrMemStreamOut : public Imf::OStream
 {
 public:
@@ -109,148 +126,273 @@ void FEXRImageWriteTask::OnAbandoned()
 	}
 }
 
-bool FEXRImageWriteTask::WriteToDisk()
+bool FEXRImageWriteTask::GenerateFilePartsForResolution(const FIntPoint& Resolution, const TArray<FImagePixelData*>& InLayers, const FString& InPartName, FEXRPart& OutPart, TArray<TUniquePtr<FImagePixelData>>& OutQuantizedLayers)
 {
-	// Ensure that the payload filename has the correct extension for the format
-	const TCHAR* FormatExtension = TEXT(".exr");
-	if (FormatExtension && !Filename.EndsWith(FormatExtension))
+	bool bSuccess = false;
+
+	OutPart.Scanlines = bPadToDataWindowSize ? Height : Resolution.Y;
+	
+	const bool bIsCropRectValid = !CropRectangle.IsEmpty() &&
+		CropRectangle.Min.X >= 0 && CropRectangle.Min.Y >= 0 &&
+		CropRectangle.Max.X <= Width && CropRectangle.Max.Y <= Height;
+
+	// Display window is always the size of the crop rectangle if one is provided, otherwise the resolution of the EXR file
+	const FIntPoint DisplayWindowRes = bIsCropRectValid ? CropRectangle.Size() : FIntPoint(Width, Height);
+
+	// When overscan/crop is provided, offset the data window into the negative region
+	const FIntPoint DataWindowOffset = bIsCropRectValid ? CropRectangle.Min : FIntPoint::ZeroValue;
+
+	// If the part's resolution does not match the file's overall resolution, we center the layers' pixels in the
+	// overall file, so offset the data window accordingly
+	const FIntPoint ResolutionDiffOffset = bPadToDataWindowSize ? FIntPoint::ZeroValue : (FIntPoint(Width, Height) - Resolution) / 2;
+	const FIntPoint DataWindowResolution = bPadToDataWindowSize ? FIntPoint(Width, Height) : Resolution;
+	
+	IMATH_NAMESPACE::V2i DataWindowTopLeft = IMATH_NAMESPACE::V2i(-DataWindowOffset.X + ResolutionDiffOffset.X, -DataWindowOffset.Y + ResolutionDiffOffset.Y);
+	IMATH_NAMESPACE::V2i DataWindowBottomRight = IMATH_NAMESPACE::V2i(DataWindowTopLeft.x + DataWindowResolution.X - 1, DataWindowTopLeft.y + DataWindowResolution.Y - 1);
+	
+	// Data Window specifies how much data is in the actual file, ie: 1920x1080
+	IMATH_NAMESPACE::Box2i DataWindow = IMATH_NAMESPACE::Box2i(DataWindowTopLeft, DataWindowBottomRight);
+
+	// Display Window specifies the total 'visible' area of the output file. 
+	// The Display Window always starts at 0,0, but Data Window can go negative to
+	// support having pixels out of bounds (such as camera overscan).
+	IMATH_NAMESPACE::Box2i DisplayWindow = IMATH_NAMESPACE::Box2i(IMATH_NAMESPACE::V2i(0, 0),IMATH_NAMESPACE::V2i(DisplayWindowRes.X - 1, DisplayWindowRes.Y - 1));
+
+	static_assert(static_cast<uint8>(EEXRCompressionFormat::Max) == Imf::Compression::NUM_COMPRESSION_METHODS);
+	Imf::Compression FileCompression = static_cast<Imf::Compression>(Compression);
+
+	OutPart.Header = Imf::Header(DisplayWindow, DataWindow, 1, Imath::V2f(0, 0), 1, Imf::LineOrder::INCREASING_Y, FileCompression);
+
+	// Multi-part headers need unique names and need their type specified
+	if (bMultipart)
 	{
-		Filename = FPaths::GetBaseFilename(Filename, false) + FormatExtension;
+		OutPart.Header.setName(TCHAR_TO_ANSI(*InPartName));
+		OutPart.Header.setType("scanlineimage");
+	}
+	
+	// If using lossy compression, specify the compression level in the header per exr spec.
+	if (FileCompression == Imf::Compression::DWAA_COMPRESSION ||
+		FileCompression == Imf::Compression::DWAB_COMPRESSION)
+	{
+		FileMetadata.Add("dwaCompressionLevel", CompressionLevel);
 	}
 
-	bool bSuccess = EnsureWritableFile();
+	// Insert our key-value pair metadata (if any, can be an arbitrary set of key/value pairs)
+	AddFileMetadata(OutPart.Header);
 
-	if (bSuccess)
+	if (ColorSpaceChromaticities.Num() > 0)
 	{
-		PreProcess();
-
-		static_assert(static_cast<uint8>(EEXRCompressionFormat::Max) == Imf::Compression::NUM_COMPRESSION_METHODS);
-		Imf::Compression FileCompression = static_cast<Imf::Compression>(Compression);
-		
-		// Data Window specifies how much data is in the actual file, ie: 1920x1080
-		IMATH_NAMESPACE::Box2i DataWindow = IMATH_NAMESPACE::Box2i(IMATH_NAMESPACE::V2i(0,0), IMATH_NAMESPACE::V2i(Width - 1, Height - 1));
-		
-		const bool bIsCropRectValid = !CropRectangle.IsEmpty() &&
-			CropRectangle.Min.X >= 0 && CropRectangle.Min.Y >= 0 &&
-			CropRectangle.Max.X <= Width && CropRectangle.Max.Y <= Height;
-		
-		// Taking Overscan into account.
-		IMATH_NAMESPACE::V2i TopLeft = IMATH_NAMESPACE::V2i(bIsCropRectValid ? CropRectangle.Min.X : 0, bIsCropRectValid ? CropRectangle.Min.Y : 0);
-		IMATH_NAMESPACE::V2i BottomRight = IMATH_NAMESPACE::V2i(bIsCropRectValid ? CropRectangle.Max.X - 1 : Width - 1, bIsCropRectValid ? CropRectangle.Max.Y - 1 : Height - 1);
-
-		// Display Window specifies the total 'visible' area of the output file. 
-		// The Display Window always starts at 0,0, but Data Window can go negative to
-		// support having pixels out of bounds (such as camera overscan).
-		IMATH_NAMESPACE::Box2i DisplayWindow = IMATH_NAMESPACE::Box2i(TopLeft, BottomRight);
-		
-		Imf::Header Header(DisplayWindow, DataWindow, 1, Imath::V2f(0, 0), 1, Imf::LineOrder::INCREASING_Y, FileCompression);
-		
-		// If using lossy compression, specify the compression level in the header per exr spec.
-		if (FileCompression == Imf::Compression::DWAA_COMPRESSION ||
-			FileCompression == Imf::Compression::DWAB_COMPRESSION)
+		if (ensureMsgf(ColorSpaceChromaticities.Num() == 4, TEXT("Four chromaticity coordinates expected.")))
 		{
-			FileMetadata.Add("dwaCompressionLevel", CompressionLevel);
+			Imf::Chromaticities Chromaticities = {
+				IMATH_NAMESPACE::V2f((float)ColorSpaceChromaticities[0].X, (float)ColorSpaceChromaticities[0].Y),
+				IMATH_NAMESPACE::V2f((float)ColorSpaceChromaticities[1].X, (float)ColorSpaceChromaticities[1].Y),
+				IMATH_NAMESPACE::V2f((float)ColorSpaceChromaticities[2].X, (float)ColorSpaceChromaticities[2].Y),
+				IMATH_NAMESPACE::V2f((float)ColorSpaceChromaticities[3].X, (float)ColorSpaceChromaticities[3].Y),
+			};
+			Imf::addChromaticities(OutPart.Header, Chromaticities);
 		}
+	}
 
-		// Insert our key-value pair metadata (if any, can be an arbitrary set of key/value pairs)
-		AddFileMetadata(Header);
+	// If we have to quantize the data (ie: Upscale 8 bit to 16 bit) we need to store them long enough for the file to get written.
+	for (FImagePixelData* Layer : InLayers)
+	{
+		uint8 RawBitDepth = Layer->GetBitDepth();
+		check(RawBitDepth == 8 || RawBitDepth == 16 || RawBitDepth == 32);
 
-		if (ColorSpaceChromaticities.Num() > 0)
+		if (!bPadToDataWindowSize && !ensureAlwaysMsgf(Layer->GetSize().X == Resolution.X && Layer->GetSize().Y == Resolution.Y, TEXT("Skipping layer due to mismatched width/height from rest of EXR file!")))
 		{
-			if (ensureMsgf(ColorSpaceChromaticities.Num() == 4, TEXT("Four chromaticity coordinates expected.")))
-			{
-				Imf::Chromaticities Chromaticities = {
-					IMATH_NAMESPACE::V2f((float)ColorSpaceChromaticities[0].X, (float)ColorSpaceChromaticities[0].Y),
-					IMATH_NAMESPACE::V2f((float)ColorSpaceChromaticities[1].X, (float)ColorSpaceChromaticities[1].Y),
-					IMATH_NAMESPACE::V2f((float)ColorSpaceChromaticities[2].X, (float)ColorSpaceChromaticities[2].Y),
-					IMATH_NAMESPACE::V2f((float)ColorSpaceChromaticities[3].X, (float)ColorSpaceChromaticities[3].Y),
-				};
-				Imf::addChromaticities(Header, Chromaticities);
-			}
+			continue;
 		}
 
-		FExrMemStreamOut OutputFile; 
-		
+		void const* RawDataPtr = nullptr;
+		int64 RawDataSize;
+		bSuccess = Layer->GetRawData(RawDataPtr, RawDataSize);
+		if (!bSuccess)
 		{
-			// The FrameBuffer stores all the channels of the resulting image.
-			Imf::FrameBuffer FrameBuffer;
+			UE_LOG(LogMovieRenderPipelineIO, Error, TEXT("Failed to retrieve raw data from image data for writing. Bailing."));
+			break;
+		}
 
-			// If we have to quantize the data (ie: Upscale 8 bit to 16 bit) we need to store them long enough for the file to get written.
-			TArray<TUniquePtr<FImagePixelData>> QuantizedData;
+		const bool bNeedsInflation = Layer->GetSize().X < Width || Layer->GetSize().Y < Height;
+		if (bPadToDataWindowSize && bNeedsInflation)
+		{
+			const FIntPoint InflateOffset = (Layer->GetSize() - FIntPoint(Width, Height)) / 2;
+			const FIntRect InflateRect = FIntRect(InflateOffset.X, InflateOffset.Y, InflateOffset.X + Width, InflateOffset.Y + Height);
 
-			for (TUniquePtr<FImagePixelData>& Layer : Layers)
+			switch (Layer->GetType())
 			{
-				uint8 RawBitDepth = Layer->GetBitDepth();
-				check(RawBitDepth == 8 || RawBitDepth == 16 || RawBitDepth == 32);
-
-				if (!ensureAlwaysMsgf(Layer->GetSize().X == Width && Layer->GetSize().Y == Height, TEXT("Skipping layer due to mismatched width/height from rest of EXR file!")))
+			case EImagePixelType::Color:
 				{
-					continue;
+					TAsyncCropImage<FColor> CropImagePreprocessor(InflateRect);
+					CropImagePreprocessor(Layer);
+					OutQuantizedLayers.Add(MoveTemp(CropImagePreprocessor.OutCroppedImage));
 				}
+				break;
 
-				void const* RawDataPtr = nullptr;
-				int64 RawDataSize;
-				bSuccess = Layer->GetRawData(RawDataPtr, RawDataSize);
-				if (!bSuccess)
+			case EImagePixelType::Float16:
 				{
-					UE_LOG(LogMovieRenderPipelineIO, Error, TEXT("Failed to retrieve raw data from image data for writing. Bailing."));
-					break;
+					TAsyncCropImage<FFloat16Color> CropImagePreprocessor(InflateRect);
+					CropImagePreprocessor(Layer);
+					OutQuantizedLayers.Add(MoveTemp(CropImagePreprocessor.OutCroppedImage));
 				}
+				break;
 
-				int64 BytesWritten = 0;
-				switch (RawBitDepth)
-				{
-				case 8:
+			case EImagePixelType::Float32:
 				{
-					TUniquePtr<FImagePixelData> QuantizedPixelData = UE::MoviePipeline::QuantizeImagePixelDataToBitDepth(Layer.Get(), 16);
-					QuantizedData.Add(MoveTemp(QuantizedPixelData));
+					TAsyncCropImage<FLinearColor> CropImagePreprocessor(InflateRect);
+					CropImagePreprocessor(Layer);
+					OutQuantizedLayers.Add(MoveTemp(CropImagePreprocessor.OutCroppedImage));
+				}
+				break;
+			}
+			
+			FString LayerName = LayerNames.FindOrAdd(Layer);
+			if (LayerName.Len() > 0)
+			{
+				LayerNames.Add(OutQuantizedLayers.Last().Get(), LayerName);
+			}
+			
+			Layer = OutQuantizedLayers.Last().Get();
+		}
+		
+		switch (RawBitDepth)
+		{
+		case 8:
+		{
+			TUniquePtr<FImagePixelData> QuantizedPixelData = UE::MoviePipeline::QuantizeImagePixelDataToBitDepth(Layer, 16);
+			OutQuantizedLayers.Add(MoveTemp(QuantizedPixelData));
 
-					// Add an entry in the LayerNames table if needed since it matches by Layer pointer but that has changed.
-					FString LayerName = LayerNames.FindOrAdd(Layer.Get());
-					if (LayerName.Len() > 0)
-					{
-						LayerNames.Add(QuantizedData.Last().Get(), LayerName);
-					}
+			// Add an entry in the LayerNames table if needed since it matches by Layer pointer but that has changed.
+			FString LayerName = LayerNames.FindOrAdd(Layer);
+			if (LayerName.Len() > 0)
+			{
+				LayerNames.Add(OutQuantizedLayers.Last().Get(), LayerName);
+			}
 
-					BytesWritten = CompressRaw<Imf::HALF>(Header, FrameBuffer, QuantizedData.Last().Get());
-				}
-					break;
-				case 16:
-					BytesWritten = CompressRaw<Imf::HALF>(Header, FrameBuffer, Layer.Get());
-					break;
-				case 32:
-					BytesWritten = CompressRaw<Imf::FLOAT>(Header, FrameBuffer, Layer.Get());
+			OutPart.BytesWritten = CompressRaw<Imf::HALF>(OutPart.Header, OutPart.FrameBuffer, OutQuantizedLayers.Last().Get());
+		}
+			break;
+		case 16:
+			OutPart.BytesWritten = CompressRaw<Imf::HALF>(OutPart.Header, OutPart.FrameBuffer, Layer);
+			break;
+		case 32:
+			OutPart.BytesWritten = CompressRaw<Imf::FLOAT>(OutPart.Header, OutPart.FrameBuffer, Layer);
+			break;
+		default:
+			checkNoEntry();
+		}
+	}
+
+	return bSuccess;
+}
+
+bool FEXRImageWriteTask::WriteToDisk()
+{
+	// Ensure that the payload filename has the correct extension for the format
+	const TCHAR* FormatExtension = TEXT(".exr");
+	if (FormatExtension && !Filename.EndsWith(FormatExtension))
+	{
+		Filename = FPaths::GetBaseFilename(Filename, false) + FormatExtension;
+	}
+
+	bool bSuccess = EnsureWritableFile();
+
+	if (bSuccess)
+	{
+		PreProcess();
+
+		FExrMemStreamOut OutputFile;
+		TArray<FEXRPart> Parts;
+		
+		// If we have to quantize the data (ie: Upscale 8 bit to 16 bit) we need to store them long enough for the file to get written.
+		TArray<TUniquePtr<FImagePixelData>> QuantizedLayers;
+		
+		if (bMultipart)
+		{
+			for (TUniquePtr<FImagePixelData>& Layer : Layers)
+			{
+				FEXRPart Part;
+				FString PartName = LayerNames.Contains(Layer.Get()) ? LayerNames[Layer.Get()] : TEXT("FinalImage");
+				if (!GenerateFilePartsForResolution(Layer->GetSize(), { Layer.Get() }, PartName, Part, QuantizedLayers))
+				{
+					UE_LOG(LogMovieRenderPipelineIO, Error, TEXT("There was an error writing layer '%s' to the EXR file. The pixel format may not be compatible with this image type, or there was a resolution mismatch."), *LayerNames[Layer.Get()]);
+					bSuccess = false;
 					break;
-				default:
-					checkNoEntry();
 				}
 				
-				// Reserve enough space in the output file for the whole layer so we don't keep reallocating.
-				OutputFile.Data.Reserve(BytesWritten);
+				Parts.Add(Part);
+			}
+		}
+		else
+		{
+			FEXRPart Part;
+			TArray<FImagePixelData*> LayerRawPtrs;
+			Algo::Transform(Layers, LayerRawPtrs, [](const TUniquePtr<FImagePixelData>& Layer) { return Layer.Get(); });
+			if (!GenerateFilePartsForResolution(FIntPoint(Width, Height), LayerRawPtrs, TEXT("FinalImage"), Part, QuantizedLayers))
+			{
+				bSuccess = false;
 			}
 
-			// This scope ensures that IMF::Outputfile creates a complete file by closing the file when it goes out of scope.
-			// To complete the file, EXR seeks back into the file and writes the scanline offsets when the file is closed, 
-			// which moves the tellp location. So file length is stored in advance for later use. The output file needs to be
-			// created after the header information is filled.
-			//
-			// Note: OutputFile has an option to control the number of threads used to write the file. The default is fine;
-			// providing too many threads here will massively decrease performance and bloat memory usage.
-			Imf::OutputFile ImfFile(OutputFile, Header);
+			Parts.Add(Part);
+		}
+		
+		int64 TotalBytesWritten = Algo::Accumulate(Parts, 0, [](int64 Sum, const FEXRPart& Part) { return Sum + Part.BytesWritten; });
+		OutputFile.Data.Reserve(TotalBytesWritten);
+		
+		if (Parts.Num() == 0)
+		{
+			UE_LOG(LogMovieRenderPipelineIO, Error, TEXT("Failed to write any layers to EXR file. The pixel format may not be compatible with this image type, or there was a resolution mismatch."));
+			return false;
+		}
+		
+		// This scope ensures that IMF::Outputfile creates a complete file by closing the file when it goes out of scope.
+		// To complete the file, EXR seeks back into the file and writes the scanline offsets when the file is closed, 
+		// which moves the tellp location. So file length is stored in advance for later use. The output file needs to be
+		// created after the header information is filled.
+		//
+		// Note: OutputFile has an option to control the number of threads used to write the file. The default is fine;
+		// providing too many threads here will massively decrease performance and bloat memory usage.
 #if WITH_EDITOR
-			try
+		try
 #endif
+		{
+			if (bMultipart)
 			{
-				ImfFile.setFrameBuffer(FrameBuffer);
-				ImfFile.writePixels(Height);
+				TArray<Imf::Header> Headers;
+				Algo::Transform(Parts, Headers, [](const FEXRPart& Part) { return Part.Header; });
+		
+				Imf::Header* HeaderData = Headers.GetData();
+				const int32 HeaderCount = Headers.Num();
+		
+				Imf::MultiPartOutputFile ImfMultiPartFile(OutputFile, HeaderData, HeaderCount);
+				for (int32 PartIdx = 0; PartIdx < Parts.Num(); ++PartIdx)
+				{
+					FEXRPart& Part = Parts[PartIdx];
+					
+					Imf::OutputPart OutputPart(ImfMultiPartFile, PartIdx);
+					OutputPart.setFrameBuffer(Part.FrameBuffer);
+					OutputPart.writePixels(Part.Scanlines);
+				}
 			}
-#if WITH_EDITOR
-			catch (const IEX_NAMESPACE::BaseExc& Exception)
+			else
 			{
-				UE_LOG(LogMovieRenderPipelineIO, Error, TEXT("Caught exception: %hs"), Exception.message().c_str());
+				if (Parts.Num() > 1)
+				{
+					UE_LOG(LogMovieRenderPipelineIO, Warning, TEXT("Multiple headers were created which is only supported for multi-part EXR files, only the first header will be used"));
+				}
+				
+				Imf::OutputFile ImfFile(OutputFile, Parts[0].Header);
+				ImfFile.setFrameBuffer(Parts[0].FrameBuffer);
+				ImfFile.writePixels(Parts[0].Scanlines);
 			}
-#endif
 		}
+#if WITH_EDITOR
+		catch (const IEX_NAMESPACE::BaseExc& Exception)
+		{
+			UE_LOG(LogMovieRenderPipelineIO, Error, TEXT("Caught exception: %hs"), Exception.message().c_str());
+		}
+#endif
 		
 		// Now that the scope has closed for the Imf::OutputFile, now we can write the data to disk.
 		if (bSuccess)
@@ -341,6 +483,8 @@ int64 FEXRImageWriteTask::CompressRaw(Imf::Header& InHeader, Imf::FrameBuffer& I
 	FString& LayerName = LayerNames.FindOrAdd(InLayer);
 	int32 NumChannels = InLayer->GetNumChannels();
 	int32 ComponentWidth = GetComponentWidth(InLayer->GetType());
+	int32 LayerWidth = InLayer->GetSize().X;
+	int32 LayerHeight = InLayer->GetSize().Y;
 
 	for (int32 Channel = 0; Channel < NumChannels; Channel++)
 	{
@@ -348,16 +492,18 @@ int64 FEXRImageWriteTask::CompressRaw(Imf::Header& InHeader, Imf::FrameBuffer& I
 
 		// Insert the channel into the header with the right datatype.
 		InHeader.channels().insert(TCHAR_TO_ANSI(*ChannelName), Imf::Channel(OutputFormat));
+		IMATH_NAMESPACE::Box2i& DataWindow = InHeader.dataWindow();
 
 		// Now insert the data for this channel. Unreal stores them interleaved.
 		InFrameBuffer.insert(TCHAR_TO_ANSI(*ChannelName),	// Name
-			Imf::Slice(OutputFormat,						// Type
+			Imf::Slice::Make(OutputFormat,						// Type
 			(char*)RawDataPtr + (ComponentWidth * Channel), // Data Start (offset by component to match interleave)
+				DataWindow,
 				ComponentWidth * NumChannels,				// xStride
-				Width * ComponentWidth * NumChannels));		// yStride
+				LayerWidth * ComponentWidth * NumChannels));		// yStride
 	}
 	
-	return int64(Width) * int64(Height) * NumChannels * int64(OutputFormat == 2 ? 4 : 2);
+	return int64(LayerWidth) * int64(LayerHeight) * NumChannels * int64(OutputFormat == 2 ? 4 : 2);
 }
 
 bool FEXRImageWriteTask::EnsureWritableFile()
@@ -544,165 +690,143 @@ void UMoviePipelineImageSequenceOutput_EXR::OnReceiveImageDataImpl(FMoviePipelin
 	UMoviePipelineOutputSetting* OutputSettings = GetPipeline()->GetPipelinePrimaryConfig()->FindSetting<UMoviePipelineOutputSetting>();
 	check(OutputSettings);
 
-	// EXR only supports one resolution per file, but in certain scenarios we can get layers with different resolutions. To solve this, we're
-	// going to write one exr file per image resolution. First we loop through all layers to figure out how many sizes we're dealing with.
-	TArray<FIntPoint> Resolutions;
+	
+	// Find the maximum resolution over all layers, which will be used to pad lower resolution layers to a matching size
+	FIntPoint MaximumResolution = FIntPoint::ZeroValue;
 	for (TPair<FMoviePipelinePassIdentifier, TUniquePtr<FImagePixelData>>& RenderPassData : InMergedOutputFrame->ImageOutputData)
 	{
-		Resolutions.AddUnique(RenderPassData.Value->GetSize());
+		MaximumResolution.X = FMath::Max(MaximumResolution.X, RenderPassData.Value->GetSize().X);
+		MaximumResolution.Y = FMath::Max(MaximumResolution.Y, RenderPassData.Value->GetSize().Y);
 	}
 
-	// Then submit multiple write tasks. Layers that don't match the current resolution will be skipped until the correct iteration of the loop.
-	for (int32 Index = 0; Index < Resolutions.Num(); Index++)
+	FString OutputDirectory = OutputSettings->OutputDirectory.Path;
+
+	// We need to resolve the filename format string. We combine the folder and file name into one long string first
+	FString FinalFilePath;
+	FMoviePipelineFormatArgs FinalFormatArgs;
+	FString FinalImageSequenceFileName;
+	FString ClipName;
+	const TCHAR* Extension = TEXT("exr");
 	{
-		FString OutputDirectory = OutputSettings->OutputDirectory.Path;
+		FString FileNameFormatString = OutputSettings->FileNameFormat;
 
-		// We need to resolve the filename format string. We combine the folder and file name into one long string first
-		FString FinalFilePath;
-		FMoviePipelineFormatArgs FinalFormatArgs;
-		FString FinalImageSequenceFileName;
-		FString ClipName;
-		const TCHAR* Extension = TEXT("exr");
-		{
-			// If we have more than one resolution we'll store it as "_Add" / "_Add(1)" etc.
-			FString FileNameFormatString = OutputSettings->FileNameFormat + "{ExtraTag}";
+		// If we're writing more than one render pass out, we need to ensure the file name has the format string in it so we don't
+		// overwrite the same file multiple times. Burn In overlays don't count because they get composited on top of an existing file.
+		const bool bIncludeRenderPass = false;
+		const bool bTestFrameNumber = true;
 
-			// If we're writing more than one render pass out, we need to ensure the file name has the format string in it so we don't
-			// overwrite the same file multiple times. Burn In overlays don't count because they get composited on top of an existing file.
-			const bool bIncludeRenderPass = false;
-			const bool bTestFrameNumber = true;
+		UE::MoviePipeline::ValidateOutputFormatString(FileNameFormatString, bIncludeRenderPass, bTestFrameNumber);
 
-			UE::MoviePipeline::ValidateOutputFormatString(FileNameFormatString, bIncludeRenderPass, bTestFrameNumber);
+		// Create specific data that needs to override 
+		TMap<FString, FString> FormatOverrides;
+		FormatOverrides.Add(TEXT("render_pass"), TEXT("")); // Render Passes are included inside the exr file by named layers.
+		FormatOverrides.Add(TEXT("ext"), Extension);
 
-			// Create specific data that needs to override 
-			TMap<FString, FString> FormatOverrides;
-			FormatOverrides.Add(TEXT("render_pass"), TEXT("")); // Render Passes are included inside the exr file by named layers.
-			FormatOverrides.Add(TEXT("ext"), Extension);
+		// This resolves the filename format and gathers metadata from the settings at the same time.
+		GetPipeline()->ResolveFilenameFormatArguments(FileNameFormatString, FormatOverrides, FinalImageSequenceFileName, FinalFormatArgs, &InMergedOutputFrame->FrameOutputState, -InMergedOutputFrame->FrameOutputState.ShotOutputFrameNumber);
 
-			// The logic for the ExtraTag is a little complicated. If there's only one layer (ideal situation) then it's empty.
-			if (Index == 0)
-			{
-				FormatOverrides.Add(TEXT("ExtraTag"), TEXT(""));
-			}
-			else if(Resolutions.Num() == 2)
-			{
-				// This is our most common case when we have a second file (the only expected one really)
-				FormatOverrides.Add(TEXT("ExtraTag"), TEXT("_Add"));
-			}
-			else
-			{
-				// Finally a fallback in the event we have three or more unique resolutions.
-				FormatOverrides.Add(TEXT("ExtraTag"), FString::Printf(TEXT("_Add(%d)"), Index));
-			}
+		FString FilePathFormatString = OutputDirectory / FileNameFormatString;
+		GetPipeline()->ResolveFilenameFormatArguments(FilePathFormatString, FormatOverrides, FinalFilePath, FinalFormatArgs, &InMergedOutputFrame->FrameOutputState);
 
-			// This resolves the filename format and gathers metadata from the settings at the same time.
-			GetPipeline()->ResolveFilenameFormatArguments(FileNameFormatString, FormatOverrides, FinalImageSequenceFileName, FinalFormatArgs, &InMergedOutputFrame->FrameOutputState, -InMergedOutputFrame->FrameOutputState.ShotOutputFrameNumber);
-
-			FString FilePathFormatString = OutputDirectory / FileNameFormatString;
-			GetPipeline()->ResolveFilenameFormatArguments(FilePathFormatString, FormatOverrides, FinalFilePath, FinalFormatArgs, &InMergedOutputFrame->FrameOutputState);
+		if (FPaths::IsRelative(FinalFilePath))
+		{
+			FinalFilePath = FPaths::ConvertRelativePathToFull(FinalFilePath);
+		}
 
-			if (FPaths::IsRelative(FinalFilePath))
-			{
-				FinalFilePath = FPaths::ConvertRelativePathToFull(FinalFilePath);
-			}
+		// Create a deterministic clipname by removing frame numbers, file extension, and any trailing .'s
+		UE::MoviePipeline::RemoveFrameNumberFormatStrings(FileNameFormatString, true);
+		GetPipeline()->ResolveFilenameFormatArguments(FileNameFormatString, FormatOverrides, ClipName, FinalFormatArgs, &InMergedOutputFrame->FrameOutputState);
+		ClipName.RemoveFromEnd(Extension);
+		ClipName.RemoveFromEnd(".");
+	}
 
-			// Create a deterministic clipname by removing frame numbers, file extension, and any trailing .'s
-			UE::MoviePipeline::RemoveFrameNumberFormatStrings(FileNameFormatString, true);
-			GetPipeline()->ResolveFilenameFormatArguments(FileNameFormatString, FormatOverrides, ClipName, FinalFormatArgs, &InMergedOutputFrame->FrameOutputState);
-			ClipName.RemoveFromEnd(Extension);
-			ClipName.RemoveFromEnd(".");
-		}
+	// If not using multi-part, we have to pad all layers up to the maximum resolution. If multi-part is on, different header
+	// data window sizes are suppported, so check the cvar to see if we should pad
+	const bool bPadToDataWindowSize = !bMultipart || UE::MoviePipeline::CVarMoviePipelinePadLayersForMultiPartEXR.GetValueOnGameThread();
+	TUniquePtr<FEXRImageWriteTask> MultiLayerImageTask = MakeUnique<FEXRImageWriteTask>();
+	MultiLayerImageTask->Filename = FinalFilePath;
+	MultiLayerImageTask->bMultipart = bMultipart;
+	MultiLayerImageTask->bPadToDataWindowSize = bPadToDataWindowSize;
+	MultiLayerImageTask->Compression = Compression;
+	// MultiLayerImageTask->CompressionLevel is intentionally skipped because it doesn't seem to make any practical difference
+	// so we don't expose it to the user because that will just cause confusion where the setting doesn't seem to do anything.
+
+	// FinalFormatArgs.FileMetadata has been merged by ResolveFilenameFormatArgs with the FrameOutputState,
+	// but we need to convert from FString, FString (needed for BP/Python purposes) to a FStringFormatArg as
+	// we need to preserve numeric metadata types later in the image writing process (for compression level)
+	TMap<FString, FStringFormatArg> NewFileMetdataMap;
+	for (const TPair<FString, FString>& Metadata : FinalFormatArgs.FileMetadata)
+	{
+		NewFileMetdataMap.Add(Metadata.Key, Metadata.Value);
+	}
+	MultiLayerImageTask->FileMetadata = NewFileMetdataMap;
 
-		TUniquePtr<FEXRImageWriteTask> MultiLayerImageTask = MakeUnique<FEXRImageWriteTask>();
-		MultiLayerImageTask->Filename = FinalFilePath;
-		MultiLayerImageTask->Compression = Compression;
-		// MultiLayerImageTask->CompressionLevel is intentionally skipped because it doesn't seem to make any practical difference
-		// so we don't expose it to the user because that will just cause confusion where the setting doesn't seem to do anything.
-
-		// FinalFormatArgs.FileMetadata has been merged by ResolveFilenameFormatArgs with the FrameOutputState,
-		// but we need to convert from FString, FString (needed for BP/Python purposes) to a FStringFormatArg as
-		// we need to preserve numeric metadata types later in the image writing process (for compression level)
-		TMap<FString, FStringFormatArg> NewFileMetdataMap;
-		for (const TPair<FString, FString>& Metadata : FinalFormatArgs.FileMetadata)
-		{
-			NewFileMetdataMap.Add(Metadata.Key, Metadata.Value);
-		}
-		MultiLayerImageTask->FileMetadata = NewFileMetdataMap;
+	// Add color space metadata to the output: xy chromaticity coordinates and/or the color space source/dest names.
+	// TODO: Support is also needed for regular exrs via the image wrapper module.
+	UMoviePipelineColorSetting* ColorSetting = GetPipeline()->GetPipelinePrimaryConfig()->FindSetting<UMoviePipelineColorSetting>();
+	if (ColorSetting && ColorSetting->OCIOConfiguration.bIsEnabled)
+	{
+		UE::MoviePipeline::UpdateColorSpaceMetadata(ColorSetting->OCIOConfiguration.ColorConfiguration, *MultiLayerImageTask);
+	}
+	else
+	{
+		ESceneCaptureSource SceneCaptureSource = (ColorSetting && ColorSetting->bDisableToneCurve) ? ESceneCaptureSource::SCS_FinalColorHDR : ESceneCaptureSource::SCS_FinalToneCurveHDR;
+		UE::MoviePipeline::UpdateColorSpaceMetadata(SceneCaptureSource, *MultiLayerImageTask);
+	}
 
-		// Add color space metadata to the output: xy chromaticity coordinates and/or the color space source/dest names.
-		// TODO: Support is also needed for regular exrs via the image wrapper module.
-		UMoviePipelineColorSetting* ColorSetting = GetPipeline()->GetPipelinePrimaryConfig()->FindSetting<UMoviePipelineColorSetting>();
-		if (ColorSetting && ColorSetting->OCIOConfiguration.bIsEnabled)
+	int32 LayerIndex = 0;
+	bool bRequiresTransparentOutput = false;
+	int32 ShotIndex = 0;
+	for (TPair<FMoviePipelinePassIdentifier, TUniquePtr<FImagePixelData>>& RenderPassData : InMergedOutputFrame->ImageOutputData)
+	{
+		// No quantization required, just copy the data as we will move it into the image write task.
+		TUniquePtr<FImagePixelData> PixelData = RenderPassData.Value->CopyImageData();
+		FImagePixelDataPayload* Payload = RenderPassData.Value->GetPayload<FImagePixelDataPayload>();
+		ShotIndex = Payload->SampleState.OutputState.ShotIndex;
+
+		// If there is more than one layer, then we will prefix the layer. The first layer is not prefixed (and gets inserted as RGBA)
+		// as most programs that handle EXRs expect the main image data to be in an unnamed layer.
+		if (LayerIndex == 0)
 		{
-			UE::MoviePipeline::UpdateColorSpaceMetadata(ColorSetting->OCIOConfiguration.ColorConfiguration, *MultiLayerImageTask);
+			// Only check the main image pass for transparent output since that's generally considered the 'preview'.
+			bRequiresTransparentOutput = Payload->bRequireTransparentOutput;
+			MultiLayerImageTask->OverscanPercentage = Payload->SampleState.OverscanPercentage;
+			MultiLayerImageTask->CropRectangle = Payload->SampleState.CropRectangle;
 		}
 		else
 		{
-			ESceneCaptureSource SceneCaptureSource = (ColorSetting && ColorSetting->bDisableToneCurve) ? ESceneCaptureSource::SCS_FinalColorHDR : ESceneCaptureSource::SCS_FinalToneCurveHDR;
-			UE::MoviePipeline::UpdateColorSpaceMetadata(SceneCaptureSource, *MultiLayerImageTask);
-		}
-
-		int32 LayerIndex = 0;
-		bool bRequiresTransparentOutput = false;
-		int32 ShotIndex = 0;
-		for (TPair<FMoviePipelinePassIdentifier, TUniquePtr<FImagePixelData>>& RenderPassData : InMergedOutputFrame->ImageOutputData)
-		{
-			if (RenderPassData.Value->GetSize() != Resolutions[Index])
-			{
-				// If this layer isn't for this resolution, don't add it to the multilayer task, a second task will be created soon.
-				continue;
-			}
-
-			// No quantization required, just copy the data as we will move it into the image write task.
-			TUniquePtr<FImagePixelData> PixelData = RenderPassData.Value->CopyImageData();
-			FImagePixelDataPayload* Payload = RenderPassData.Value->GetPayload<FImagePixelDataPayload>();
-			ShotIndex = Payload->SampleState.OutputState.ShotIndex;
-
 			// If there is more than one layer, then we will prefix the layer. The first layer is not prefixed (and gets inserted as RGBA)
-			// as most programs that handle EXRs expect the main image data to be in an unnamed layer.
-			if (LayerIndex == 0)
+			// as most programs that handle EXRs expect the main image data to be in an unnamed layer. We only postfix with cameraname
+			// if there's multiple cameras, as pipelines may be already be built around the generic "one camera" support.
+			UMoviePipelineExecutorShot* CurrentShot = GetPipeline()->GetActiveShotList()[ShotIndex];
+			UMoviePipelineCameraSetting* CameraSettings = GetPipeline()->FindOrAddSettingForShot<UMoviePipelineCameraSetting>(CurrentShot);
+			int32 NumCameras = CameraSettings->bRenderAllCameras ? CurrentShot->SidecarCameras.Num() : 1;
+			
+			FString CombinedName;
+			if (NumCameras == 1)
 			{
-				// Only check the main image pass for transparent output since that's generally considered the 'preview'.
-				bRequiresTransparentOutput = Payload->bRequireTransparentOutput;
-				MultiLayerImageTask->OverscanPercentage = Payload->SampleState.OverscanPercentage;
-				MultiLayerImageTask->CropRectangle = Payload->SampleState.CropRectangle;
+				CombinedName = RenderPassData.Key.Name;
 			}
 			else
 			{
-				// If there is more than one layer, then we will prefix the layer. The first layer is not prefixed (and gets inserted as RGBA)
-				// as most programs that handle EXRs expect the main image data to be in an unnamed layer. We only postfix with cameraname
-				// if there's multiple cameras, as pipelines may be already be built around the generic "one camera" support.
-				UMoviePipelineExecutorShot* CurrentShot = GetPipeline()->GetActiveShotList()[ShotIndex];
-				UMoviePipelineCameraSetting* CameraSettings = GetPipeline()->FindOrAddSettingForShot<UMoviePipelineCameraSetting>(CurrentShot);
-				int32 NumCameras = CameraSettings->bRenderAllCameras ? CurrentShot->SidecarCameras.Num() : 1;
-				
-				FString CombinedName;
-				if (NumCameras == 1)
-				{
-					CombinedName = RenderPassData.Key.Name;
-				}
-				else
-				{
-					CombinedName = FString::Printf(TEXT("%s_%s"), *RenderPassData.Key.Name, *RenderPassData.Key.CameraName);
-				}
-				MultiLayerImageTask->LayerNames.FindOrAdd(PixelData.Get(), CombinedName);
+				CombinedName = FString::Printf(TEXT("%s_%s"), *RenderPassData.Key.Name, *RenderPassData.Key.CameraName);
 			}
-
-			MultiLayerImageTask->Width = Resolutions[Index].X;
-			MultiLayerImageTask->Height = Resolutions[Index].Y;
-			MultiLayerImageTask->Layers.Add(MoveTemp(PixelData));
-			LayerIndex++;
+			MultiLayerImageTask->LayerNames.FindOrAdd(PixelData.Get(), CombinedName);
 		}
 
-		MoviePipeline::FMoviePipelineOutputFutureData OutputData;
-		OutputData.Shot = GetPipeline()->GetActiveShotList()[ShotIndex];
-		OutputData.PassIdentifier = FMoviePipelinePassIdentifier(TEXT("")); // exrs put all the render passes internally so this resolves to a ""
-		OutputData.FilePath = FinalFilePath;
-		GetPipeline()->AddOutputFuture(ImageWriteQueue->Enqueue(MoveTemp(MultiLayerImageTask)), OutputData);
+		MultiLayerImageTask->Width = MaximumResolution.X;
+		MultiLayerImageTask->Height = MaximumResolution.Y;
+		MultiLayerImageTask->Layers.Add(MoveTemp(PixelData));
+		LayerIndex++;
+	}
+
+	MoviePipeline::FMoviePipelineOutputFutureData OutputData;
+	OutputData.Shot = GetPipeline()->GetActiveShotList()[ShotIndex];
+	OutputData.PassIdentifier = FMoviePipelinePassIdentifier(TEXT("")); // exrs put all the render passes internally so this resolves to a ""
+	OutputData.FilePath = FinalFilePath;
+	GetPipeline()->AddOutputFuture(ImageWriteQueue->Enqueue(MoveTemp(MultiLayerImageTask)), OutputData);
 
 #if WITH_EDITOR
-		GetPipeline()->AddFrameToOutputMetadata(ClipName, FinalImageSequenceFileName, InMergedOutputFrame->FrameOutputState, Extension, bRequiresTransparentOutput);
+	GetPipeline()->AddFrameToOutputMetadata(ClipName, FinalImageSequenceFileName, InMergedOutputFrame->FrameOutputState, Extension, bRequiresTransparentOutput);
 #endif
-
-	}
 }
diff --git a/Engine/Plugins/MovieScene/MovieRenderPipeline/Source/MovieRenderPipelineRenderPasses/Public/MovieGraphImageSequenceOutputNode.h b/Engine/Plugins/MovieScene/MovieRenderPipeline/Source/MovieRenderPipelineRenderPasses/Public/MovieGraphImageSequenceOutputNode.h
--- a/Engine/Plugins/MovieScene/MovieRenderPipeline/Source/MovieRenderPipelineRenderPasses/Public/MovieGraphImageSequenceOutputNode.h
+++ b/Engine/Plugins/MovieScene/MovieRenderPipeline/Source/MovieRenderPipelineRenderPasses/Public/MovieGraphImageSequenceOutputNode.h
@@ -150,7 +150,8 @@ protected:
 	/** Convenience function to create a new EXR image write task, given a file name and compression format. */
 	TUniquePtr<FEXRImageWriteTask> CreateImageWriteTask(
 		FString InFileName,
-		EEXRCompressionFormat InCompression
+		EEXRCompressionFormat InCompression,
+		bool bMultiPart = false
 	) const;
 	
 	/** Convenience function to prepare the image write task's global file metadata. */
@@ -186,6 +187,7 @@ public:
 	{
 		// Multi-layer default excludes {layer_name}.
 		FileNameFormat = TEXT("{sequence_name}.{frame_number}");
+		bMultipart = false;
 	}
 
 	virtual EMovieGraphBranchRestriction GetBranchRestriction() const override
@@ -229,19 +231,25 @@ public:
 #endif
 
 private:
+	/** Stores the EXR output config to use for a specific file name */
+	struct FEXROutputConfigForFilename
+	{
+		TArray<FMovieGraphRenderDataIdentifier> RenderIDs;
+		FMovieGraphResolveArgs ResolveArgs;
+		FString FrameTemplatedFilename;
+		FIntPoint MaximumResolution;
+	};
+
 	/**
-	 * Generates filenames for each render pass (renderID), placed in OutFilenameToRenderIDs. Also provides the resolve
-	 * args that were created when resolving the filename in OutFilenameToResolveArgs. The generated mapping ensures
-	 * that each filename only points to renderIDs with a common resolution, since currently EXRs can only contain layers
-	 * with the same resolution. Frame-templated filenames (ie, filenames that have placeholders in the place of frame numbers)
-	 * are also provided -- the key in OutFilenameToFrameTemplatedFilename is the non-templated filename, and the value is the templated filename.
+	 * Generates an output config for each filename, which consists of a list of render IDs, the resolve args that were created
+	 * when resolving the filename, the frame-templated filenames (ie, filenames that have placeholders in the place of frame numbers),
+	 * and the maximum resolution of all the layers written to the file, which will be used to pad out lower resolution layers so that
+	 * they can all be in the same EXR file (when multi-part is not used).
 	 */
-	void GetFilenameToRenderIDMappings(
+	void GetFilenameToEXROutputConfigMappings(
 		const UMovieGraphImageSequenceOutputNode_MultiLayerEXR* InParentNode,
 		UMovieGraphPipeline* InPipeline, UE::MovieGraph::FMovieGraphOutputMergerFrame* InRawFrameData,
-		TMap<FString, TArray<FMovieGraphRenderDataIdentifier>>& OutFilenameToRenderIDs,
-		TMap<FString, FMovieGraphResolveArgs>& OutFilenameToResolveArgs,
-		TMap<FString, FString>& OutFilenameToFrameTemplatedFilename) const;
+		TMap<FString, FEXROutputConfigForFilename>& OutFilenameToOutputConfigs) const;
 	
 	/**
 	 * Generates the filename that the EXR will be written to, as well as the resulting resolve args via OutResolveArgs.
@@ -249,9 +257,21 @@ private:
 	 */
 	FString ResolveOutputFilename(
 		const UMovieGraphImageSequenceOutputNode_MultiLayerEXR* InParentNode,
-		const UMovieGraphPipeline* InPipeline, const int32 ResolutionIndex, const UE::MovieGraph::FMovieGraphOutputMergerFrame* InRawFrameData,
+		const UMovieGraphPipeline* InPipeline, const UE::MovieGraph::FMovieGraphOutputMergerFrame* InRawFrameData,
 		const FMovieGraphRenderDataIdentifier& InRenderDataIdentifier, FMovieGraphResolveArgs& OutResolveArgs,
 		FString& OutFrameTemplatedFilename) const;
+	
+public:
+	UPROPERTY(EditAnywhere, BlueprintReadWrite, Category = Overrides, meta = (InlineEditConditionToggle))
+	uint8 bOverride_bMultiPart : 1;
+	
+	/**
+	 * Indicates whether the exr file should be written as a multi-part exr file, which supports having different image types and resolutions for each layer.
+	 * Multi-part EXRs are a feature of EXR 2.0 and may not be supported by all software. If the console variable 'MoviePipeline.PadLayersForMultiPartEXR' is enabled,
+	 * then all parts of the multi-part EXR will be padded to match the data window of the largest layer, as some software does not support different data window sizes.
+	 */
+	UPROPERTY(EditAnywhere, BlueprintReadWrite, Category = "EXR", meta=(EditCondition="bOverride_bMultiPart"))
+	bool bMultipart = false;
 };
 
 /**
diff --git a/Engine/Plugins/MovieScene/MovieRenderPipeline/Source/MovieRenderPipelineRenderPasses/Public/MoviePipelineEXROutput.h b/Engine/Plugins/MovieScene/MovieRenderPipeline/Source/MovieRenderPipelineRenderPasses/Public/MoviePipelineEXROutput.h
--- a/Engine/Plugins/MovieScene/MovieRenderPipeline/Source/MovieRenderPipelineRenderPasses/Public/MoviePipelineEXROutput.h
+++ b/Engine/Plugins/MovieScene/MovieRenderPipeline/Source/MovieRenderPipelineRenderPasses/Public/MoviePipelineEXROutput.h
@@ -29,6 +29,9 @@ namespace UE
 {
 	namespace MoviePipeline
 	{
+		/** CVar indicating if the EXR file writer is allowed to inflate layers of a lower resolution */
+		extern TAutoConsoleVariable<bool> CVarMoviePipelinePadLayersForMultiPartEXR;
+
 		/** Collection of color space metadata for EXR. */
 		struct FEXRColorSpaceMetadata
 		{
@@ -119,12 +122,20 @@ public:
 	/** Color space chromaticity metadata. */
 	TArray<FVector2d> ColorSpaceChromaticities;
 
+	/** Indicates if a multi-part file should be used, which allows each channel to have its own header and properties */
+	bool bMultipart;
+	
+	/** Indicates that any layers that do not match the base resolution of the EXR file should be inflated with empty pixels to that resolution */
+	bool bPadToDataWindowSize;
+	
 	FEXRImageWriteTask()
 		: bOverwriteFile(true)
 		, Compression(EEXRCompressionFormat::PIZ)
 		, CompressionLevel(45)
 		, OverscanPercentage(0.0f)
 		, CropRectangle(FIntRect())
+		, bMultipart(true)
+		, bPadToDataWindowSize(true)
 	{}
 
 public:
@@ -133,7 +144,18 @@ public:
 	virtual void OnAbandoned() override final;
 
 private:
-
+	/** EXR file parts and their metadata (scanline count, bytes written) */
+	struct FEXRPart
+	{
+		Imf::Header Header;
+		Imf::FrameBuffer FrameBuffer;
+		int32 Scanlines;
+		int64 BytesWritten;
+	};
+
+	/** Generates EXR file parts (header, framebuffer, etc) for the specified layers at the specified resolution */
+	bool GenerateFilePartsForResolution(const FIntPoint& Resolution, const TArray<FImagePixelData*>& InLayers, const FString& InPartName, FEXRPart& OutPart, TArray<TUniquePtr<FImagePixelData>>& OutQuantizedLayers);
+	
 	/**
 	 * Run the task, attempting to write out the raw data using the currently specified parameters
 	 *
@@ -177,6 +199,7 @@ public:
 		OutputFormat = EImageFormat::EXR;
 		Compression = EEXRCompressionFormat::PIZ;
 		bMultilayer = true;
+		bMultipart = false;
 	}
 
 	virtual void OnReceiveImageDataImpl(FMoviePipelineMergerOutputFrame* InMergedOutputFrame) override;
@@ -196,4 +219,12 @@ public:
 	UPROPERTY(EditAnywhere, BlueprintReadWrite, Category = "EXR")
 	bool bMultilayer;
 
+	/**
+	 * Indicates whether the exr file should be written as a multi-part exr file, which supports having different image types and resolutions for each layer
+	 * Multi-part exr files are a feature of EXR 2.0 and may not be supported by all software. If the console variable 'MoviePipeline.PadLayersForMultiPartEXR' is enabled,
+	 * then all parts of the multi-part EXR will be padded to match the data window of the largest layer, as some software does not support different data window sizes.
+	 */
+	UPROPERTY(EditAnywhere, BlueprintReadWrite, Category = "EXR")
+	bool bMultipart;
+
 };
