MoviePipeline: Several fixes/improvements for Quick Render:

* Updated the default Quick Render graph to not set scalability. Quick Render will now use the viewport's scalability setting unless the graph explicitly sets it.
* Rendering from an unsaved level is now supported.
* Fixed a bug that caused ABrush actors to be visible in renders when the Visibility flag was set.
* Level visibility within the editor is now matched in renders when the Visibility flag is set.

#virtualized

diff --git a/Engine/Plugins/MovieScene/MovieRenderPipeline/Source/MovieRenderPipelineEditor/Private/Graph/MovieGraphQuickRender.cpp b/Engine/Plugins/MovieScene/MovieRenderPipeline/Source/MovieRenderPipelineEditor/Private/Graph/MovieGraphQuickRender.cpp
--- a/Engine/Plugins/MovieScene/MovieRenderPipeline/Source/MovieRenderPipelineEditor/Private/Graph/MovieGraphQuickRender.cpp
+++ b/Engine/Plugins/MovieScene/MovieRenderPipeline/Source/MovieRenderPipelineEditor/Private/Graph/MovieGraphQuickRender.cpp
@@ -4,10 +4,12 @@
 
 #include "Camera/CameraActor.h"
 #include "Editor.h"
+#include "Engine/LevelStreaming.h"
 #include "EngineUtils.h"
 #include "LevelSequence.h"
 #include "LevelSequenceActor.h"
 #include "LevelSequenceEditorBlueprintLibrary.h"
+#include "LevelUtils.h"
 #include "Misc/MessageDialog.h"
 #include "MoviePipelinePIEExecutor.h"
 #include "MoviePipelinePostRenderFileDisplayProcessor.h"
@@ -17,9 +19,12 @@
 #include "MovieRenderPipelineSettings.h"
 #include "MovieScene.h"
 #include "ObjectEditorUtils.h"
+#include "Scalability.h"
 #include "Sections/MovieSceneCameraCutSection.h"
+#include "Sections/MovieSceneLevelVisibilitySection.h"
 #include "SequencerUtilities.h"
 #include "Tracks/MovieSceneCameraCutTrack.h"
+#include "Tracks/MovieSceneLevelVisibilityTrack.h"
 #include "Tracks/MovieSceneSubTrack.h"
 
 #include "Graph/MovieGraphRenderLayerSubsystem.h"
@@ -27,11 +32,13 @@
 #include "Graph/Nodes/MovieGraphApplyViewportLookNode.h"
 #include "Graph/Nodes/MovieGraphCollectionNode.h"
 #include "Graph/Nodes/MovieGraphDeferredPassNode.h"
+#include "Graph/Nodes/MovieGraphGlobalGameOverrides.h"
 #include "Graph/Nodes/MovieGraphGlobalOutputSettingNode.h"
 #include "Graph/Nodes/MovieGraphModifierNode.h"
 #include "Graph/Nodes/MovieGraphVideoOutputNode.h"
 
 #if WITH_EDITOR
+#include "Editor/Transactor.h"
 #include "EditorViewportClient.h"
 #include "LevelEditorViewport.h"
 #include "Selection.h"
@@ -44,6 +51,31 @@
 
 #define LOCTEXT_NAMESPACE "FMovieGraphQuickRender"
 
+/**
+ * Wraps the scope that this is created in with a dummy transaction. When going out of scope, the dummy transaction will be canceled, effectively
+ * blocking any transactions that occurred while in scope.
+ */
+class FMovieGraphTransactionBlocker
+{
+public:
+	FMovieGraphTransactionBlocker()
+	{
+#if WITH_EDITOR
+		TransactionId = GEditor->BeginTransaction(FText::GetEmpty());
+#endif
+	}
+
+	~FMovieGraphTransactionBlocker()
+	{
+#if WITH_EDITOR
+		GEditor->CancelTransaction(TransactionId);
+#endif
+	}
+
+private:
+	int32 TransactionId = 0;
+};
+
 void FMovieGraphQuickRender::BeginQuickRender(const EMovieGraphQuickRenderMode QuickRenderMode, const UMovieGraphQuickRenderModeSettings& InQuickRenderSettings)
 {
 	UMoviePipelineQueueSubsystem* Subsystem = GEditor->GetEditorSubsystem<UMoviePipelineQueueSubsystem>();
@@ -55,9 +87,6 @@ void FMovieGraphQuickRender::BeginQuickRender(const EMovieGraphQuickRenderMode Q
 		return;
 	}
 
-	// Always use the PIE executor for quick render, rather than the executor specified in the project settings for normal MRG.
-	const TSubclassOf<UMoviePipelineExecutorBase> ExecutorClass = UMoviePipelinePIEExecutor::StaticClass();
-
 	// Get/create the level sequence that should be rendered from (varies depending on the Quick Render mode).
 	SetUpRenderingLevelSequence(QuickRenderMode);
 	if (!RenderingLevelSequence.IsValid())
@@ -65,8 +94,8 @@ void FMovieGraphQuickRender::BeginQuickRender(const EMovieGraphQuickRenderMode Q
 		return;
 	}
 
-	// Also set up the utility level sequence. This is not expected to be valid for most modes, so we don't need to check that this was successfully created.
-	SetUpUtilityLevelSequence(QuickRenderMode);
+	// Also set up the utility level sequence (which is used, for example, for overriding level visibility)
+	SetUpUtilityLevelSequence();
 
 	// Don't do additional work if the selected mode doesn't have what it needs to do a render. AreModePrerequisitesMet() will output the appropriate warnings.
 	if (!AreModePrerequisitesMet(QuickRenderMode))
@@ -91,16 +120,6 @@ void FMovieGraphQuickRender::BeginQuickRender(const EMovieGraphQuickRenderMode Q
 	NewJob->Map = FSoftObjectPath(EditorWorld);
 	NewJob->JobName = FString(TEXT("QuickRender"));
 
-	// Make sure the level has been saved (unsaved levels won't be usable in PIE)
-	if (!FPackageName::IsValidLongPackageName(NewJob->Map.GetLongPackageName()))
-	{
-		FMessageDialog::Open(
-			EAppMsgType::Ok,
-			LOCTEXT("UnsavedMapWarning", "Quick Render needs a saved level to function correctly. Please save your level, then try again."));
-		
-		return;
-	}
-
 	// Update the job to use the quick render graph (either the default, or the one specified in the settings)
 	TMap<TObjectPtr<UMovieGraphConfig>, TObjectPtr<UMovieGraphConfig>> OriginalGraphToDupeMap;
 	const UMovieGraphConfig* QuickRenderGraph = GetQuickRenderGraph(InQuickRenderSettings.GraphPreset.LoadSynchronous(), OriginalGraphToDupeMap);
@@ -132,6 +151,23 @@ void FMovieGraphQuickRender::BeginQuickRender(const EMovieGraphQuickRenderMode Q
 			NewVariableAssignments->SetVariableAssignmentEnableState(GraphVariable, true);
 		}
 	}
+	
+	// Evaluate the graph. Walking the graph to determine which nodes are present in each branch (especially taking into consideration
+	// subgraphs) is difficult here, so evaluation is the best way to deal with this.
+	//
+	// Ideally we would not have to evaluate the graph here, and we could get the evaluated graph from elsewhere within the pipeline.
+	// However, that would require exposing the inner workings of the pipeline to an outside API, which was deemed a bigger negative
+	// than taking the (very small) performance hit of doing a graph evaluation here.
+	FString OutError;
+	FMovieGraphTraversalContext Context;
+	Context.Job = NewJob;
+	Context.RootGraph = TemporaryGraph.Get();
+	TemporaryEvaluatedGraph = TStrongObjectPtr(TemporaryGraph->CreateFlattenedGraph(Context, OutError));
+	if (!TemporaryEvaluatedGraph)
+	{
+		UE_LOG(LogMovieRenderPipeline, Error, TEXT("Unable to Quick Render: Could not create an evaluated graph. Reason: %s"), *OutError);
+		return;
+	}
 
 	// Update the graph as needed in order to properly apply the settings specified by the user. This will perform a graph evaluation, so the job
 	// should be fully initialized at this point.
@@ -140,17 +176,20 @@ void FMovieGraphQuickRender::BeginQuickRender(const EMovieGraphQuickRenderMode Q
 	// Cache out anything necessary before PIE starts
 	CachePrePieData(EditorWorld);
 
-	// Do the render
-	UMoviePipelinePIEExecutor* Executor = Cast<UMoviePipelinePIEExecutor>(Subsystem->RenderQueueInstanceWithExecutor(TemporaryQueue.Get(), ExecutorClass));
+	// Create the executor. Note that it's allowed to use unsaved levels (many animators work off of unsaved levels). Because Quick Render is
+	// creating a temporary queue and rendering locally, using an unsaved level is safe.
+	// Always use the PIE executor for Quick Render, rather than the executor specified in the project settings for normal MRG.
+	TemporaryExecutor = TStrongObjectPtr(NewObject<UMoviePipelinePIEExecutor>(GetTransientPackage()));
+	TemporaryExecutor->SetAllowUsingUnsavedLevels(true);
 
 	// Do any needed setup in the PIE world before rendering starts
-	Executor->OnIndividualJobStarted().AddLambda([QuickRenderMode](UMoviePipelineExecutorJob* InJob)
+	TemporaryExecutor->OnIndividualJobStarted().AddLambda([QuickRenderMode, EditorWorld, &InQuickRenderSettings](UMoviePipelineExecutorJob* InJob)
 	{
-		PerformPreRenderSetup(QuickRenderMode);
+		PerformPreRenderSetup(QuickRenderMode, EditorWorld, InQuickRenderSettings);
 	});
 
 	// View the job's rendered frame(s)/video(s) if requested
-	Executor->OnIndividualJobWorkFinished().AddLambda([&InQuickRenderSettings](FMoviePipelineOutputData OutputData)
+	TemporaryExecutor->OnIndividualJobWorkFinished().AddLambda([&InQuickRenderSettings](FMoviePipelineOutputData OutputData)
 	{
 		if (InQuickRenderSettings.PostRenderBehavior == EMovieGraphQuickRenderPostRenderActionType::PlayRenderOutput)
 		{
@@ -161,16 +200,21 @@ void FMovieGraphQuickRender::BeginQuickRender(const EMovieGraphQuickRenderMode Q
 	});
 
 	// Clean up
-	Executor->OnExecutorFinished().AddLambda([QuickRenderMode](UMoviePipelineExecutorBase* InExecutor, bool bSuccess)
+	TemporaryExecutor->OnExecutorFinished().AddLambda([QuickRenderMode](UMoviePipelineExecutorBase* InExecutor, bool bSuccess)
 	{
 		RestorePreRenderState(QuickRenderMode);
 
 		// The temp objects are part of the transient package, so manually release them when the render is finished
 		TemporaryQueue.Reset();
 		TemporaryGraph.Reset();
+		TemporaryExecutor.Reset();
+		TemporaryEvaluatedGraph.Reset();
 		RenderingLevelSequence.Reset();
 		UtilityLevelSequence.Reset();
 	});
+
+	// Do the render
+	Subsystem->RenderQueueInstanceWithExecutorInstance(TemporaryQueue.Get(), TemporaryExecutor.Get());
 }
 
 UMovieGraphConfig* FMovieGraphQuickRender::GetQuickRenderGraph(UMovieGraphConfig* InUserSpecifiedGraph, TMap<TObjectPtr<UMovieGraphConfig>, TObjectPtr<UMovieGraphConfig>>& OutOriginalToDupeMap)
@@ -276,12 +320,15 @@ void FMovieGraphQuickRender::ApplyQuickRenderUpdatesToDuplicateGraph(const UMovi
 			{
 				if (OcioConfiguration->bIsEnabled)
 				{
-					ApplyQuickRenderUpdatesToDuplicateGraph_ApplyOcio(InJob, OcioConfiguration);
+					ApplyQuickRenderUpdatesToDuplicateGraph_ApplyOcio(OcioConfiguration);
 				}
 			}
 		}
 #endif	// WITH_OCIO
 	}
+
+	// Apply viewport scalability settings if the graph didn't specify an explicit scalability setting
+	ApplyQuickRenderUpdatesToDuplicateGraph_Scalability();
 }
 
 void FMovieGraphQuickRender::ApplyQuickRenderUpdatesToDuplicateGraph_ApplyEditorVisibility(const UWorld* InEditorWorld)
@@ -289,6 +336,7 @@ void FMovieGraphQuickRender::ApplyQuickRenderUpdatesToDuplicateGraph_ApplyEditor
 	TArray<FMovieGraphActorQueryEntry> VisibleActors;
 	TArray<FMovieGraphActorQueryEntry> HiddenActors;
 	
+	// Determine the actors that are hidden and visible
 	for (FActorIterator It(InEditorWorld); It; ++It)
 	{
 		AActor* Actor = *It;
@@ -298,6 +346,14 @@ void FMovieGraphQuickRender::ApplyQuickRenderUpdatesToDuplicateGraph_ApplyEditor
 			continue;
 		}
 
+		// Don't process ABrush actors (actors derived from ABrush are probably ok). The builder brush is problematic because of how it decides to
+		// render itself. There's no API we can call to accurately determine if it's going to show up in the viewport, and they're hidden in game
+		// by default, so just skip them so they don't show up in the render (needing them to show up in renders should be rare).
+		if (Actor->GetClass() == ABrush::StaticClass())
+		{
+			continue;
+		}
+
 		if (Actor->IsHiddenEd())
 		{
 			FMovieGraphActorQueryEntry& HiddenActorEntry = HiddenActors.AddDefaulted_GetRef();
@@ -310,65 +366,48 @@ void FMovieGraphQuickRender::ApplyQuickRenderUpdatesToDuplicateGraph_ApplyEditor
 		}
 	}
 
-	const UMovieGraphCollectionNode* VisibleActorsCollectionNode = Cast<UMovieGraphCollectionNode>(InjectNodeIntoBranch(UMovieGraphCollectionNode::StaticClass(), UMovieGraphNode::GlobalsPinName));
-	const UMovieGraphCollectionNode* HiddenActorsCollectionNode = Cast<UMovieGraphCollectionNode>(InjectNodeIntoBranch(UMovieGraphCollectionNode::StaticClass(), UMovieGraphNode::GlobalsPinName));
+	auto AddNewCollectionWithVisibilityModifier = []<typename QueryType>(const FString& InOperationName, const bool bModifierShouldHide) -> QueryType*
+	{
+		const UMovieGraphCollectionNode* NewCollectionNode = Cast<UMovieGraphCollectionNode>(InjectNodeIntoBranch(UMovieGraphCollectionNode::StaticClass(), UMovieGraphNode::GlobalsPinName));
 
-	const FName VisibleActorsCollectionName(TEXT("__AUTOGEN_COLLECTION_VISIBLE_ACTORS"));
-	const FName HiddenActorsCollectionName(TEXT("__AUTOGEN_COLLECTION_HIDDEN_ACTORS"));
-	VisibleActorsCollectionNode->Collection->SetCollectionName(VisibleActorsCollectionName.ToString());
-	HiddenActorsCollectionNode->Collection->SetCollectionName(HiddenActorsCollectionName.ToString());
+		const FString NewCollectionName(TEXT("__AUTOGEN_COLLECTION_") + InOperationName);
+		NewCollectionNode->Collection->SetCollectionName(NewCollectionName);
 
-	UMovieGraphConditionGroupQuery_Actor* VisibleActorsQuery =
-		Cast<UMovieGraphConditionGroupQuery_Actor>(VisibleActorsCollectionNode->Collection->AddConditionGroup()->AddQuery(UMovieGraphConditionGroupQuery_Actor::StaticClass()));
-	VisibleActorsQuery->ActorsAndComponentsToMatch = VisibleActors;
+		QueryType* NewConditionGroupQuery = Cast<QueryType>(NewCollectionNode->Collection->AddConditionGroup()->AddQuery(QueryType::StaticClass()));
 
-	UMovieGraphConditionGroupQuery_Actor* HiddenActorsQuery =
-		Cast<UMovieGraphConditionGroupQuery_Actor>(HiddenActorsCollectionNode->Collection->AddConditionGroup()->AddQuery(UMovieGraphConditionGroupQuery_Actor::StaticClass()));
-	HiddenActorsQuery->ActorsAndComponentsToMatch = HiddenActors;
+		UMovieGraphModifierNode* NewModifierNode = Cast<UMovieGraphModifierNode>(InjectNodeIntoBranch(UMovieGraphModifierNode::StaticClass(), UMovieGraphNode::GlobalsPinName));
+		NewModifierNode->ModifierName = TEXT("__AUTOGEN_MODIFIER_") + InOperationName;
+		NewModifierNode->AddCollection(FName(NewCollectionName));
 
-	UMovieGraphModifierNode* VisibleActorsModifierNode = Cast<UMovieGraphModifierNode>(InjectNodeIntoBranch(UMovieGraphModifierNode::StaticClass(), UMovieGraphNode::GlobalsPinName));
-	UMovieGraphModifierNode* HiddenActorsModifierNode = Cast<UMovieGraphModifierNode>(InjectNodeIntoBranch(UMovieGraphModifierNode::StaticClass(), UMovieGraphNode::GlobalsPinName));
-	VisibleActorsModifierNode->ModifierName = TEXT("__AUTOGEN_MODIFIER_VISIBLE_ACTORS");
-	HiddenActorsModifierNode->ModifierName = TEXT("__AUTOGEN_MODIFIER_HIDDEN_ACTORS");
-	VisibleActorsModifierNode->AddCollection(VisibleActorsCollectionName);
-	HiddenActorsModifierNode->AddCollection(HiddenActorsCollectionName);
+		UMovieGraphRenderPropertyModifier* RenderPropertyModifier = Cast<UMovieGraphRenderPropertyModifier>(NewModifierNode->GetModifier(UMovieGraphRenderPropertyModifier::StaticClass()));
 
-	UMovieGraphRenderPropertyModifier* VisibleActorsModifier = Cast<UMovieGraphRenderPropertyModifier>(VisibleActorsModifierNode->GetModifier(UMovieGraphRenderPropertyModifier::StaticClass()));
-	UMovieGraphRenderPropertyModifier* HiddenActorsModifier = Cast<UMovieGraphRenderPropertyModifier>(HiddenActorsModifierNode->GetModifier(UMovieGraphRenderPropertyModifier::StaticClass()));
+		RenderPropertyModifier->bOverride_bIsHidden = true;
+		RenderPropertyModifier->bIsHidden = bModifierShouldHide;
 
-	VisibleActorsModifier->bOverride_bIsHidden = true;
-	VisibleActorsModifier->bIsHidden = false;
-	HiddenActorsModifier->bOverride_bIsHidden = true;
-	HiddenActorsModifier->bIsHidden = true;
+		// The caller can modify the returned condition group as needed
+		return NewConditionGroupQuery;
+	};
+
+	UMovieGraphConditionGroupQuery_Actor* VisibleActorQuery =
+		AddNewCollectionWithVisibilityModifier.operator()<UMovieGraphConditionGroupQuery_Actor>(TEXT("VISIBLE_ACTORS"), false);
+	VisibleActorQuery->ActorsAndComponentsToMatch = VisibleActors;
+	
+	UMovieGraphConditionGroupQuery_Actor* HiddenActorQuery =
+		AddNewCollectionWithVisibilityModifier.operator()<UMovieGraphConditionGroupQuery_Actor>(TEXT("HIDDEN_ACTORS"), true);
+	HiddenActorQuery->ActorsAndComponentsToMatch = HiddenActors;
 }
 
-void FMovieGraphQuickRender::ApplyQuickRenderUpdatesToDuplicateGraph_ApplyOcio(UMoviePipelineExecutorJob* InJob, const FOpenColorIODisplayConfiguration* InOcioConfiguration)
+void FMovieGraphQuickRender::ApplyQuickRenderUpdatesToDuplicateGraph_ApplyOcio(const FOpenColorIODisplayConfiguration* InOcioConfiguration)
 {
-	FMovieGraphTraversalContext Context;
-	Context.Job = InJob;
-	Context.RootGraph = TemporaryGraph.Get();
-	FString OutError;
-
-	// Evaluate the graph. This is done so we can tell which branch(es) have nodes we need to update in order to properly apply
-	// OCIO. For example, if the evaluated graph shows a deferred renderer and JPEG node on "Branch1", then update Branch1
-	// accordingly. But if "Branch2" has no deferred renderer, then there's nothing to do there.
-	//
-	// Ideally we would not have to evaluate the graph here, and we could get the evaluated graph from elsewhere within the pipeline.
-	// However, that would require exposing the inner workings of the pipeline to an outside API, which was deemed a bigger negative
-	// than taking the (very small) performance hit of doing a graph evaluation here.
-	UMovieGraphEvaluatedConfig* FlattenedConfig = TemporaryGraph->CreateFlattenedGraph(Context, OutError);
-
-	const TArray<FName> BranchNames = FlattenedConfig->GetBranchNames();
-
 	// For each branch, apply OCIO configuration.
-	for (const FName& BranchName : BranchNames)
+	for (const FName& BranchName : TemporaryEvaluatedGraph->GetBranchNames())
 	{
 		constexpr bool bIncludeCDOs = false;
 		constexpr bool bExactMatch = false;
 		const TArray<UMovieGraphFileOutputNode*> FileOutputNodes =
-			FlattenedConfig->GetSettingsForBranch<UMovieGraphFileOutputNode>(BranchName, bIncludeCDOs, bExactMatch);
+			TemporaryEvaluatedGraph->GetSettingsForBranch<UMovieGraphFileOutputNode>(BranchName, bIncludeCDOs, bExactMatch);
 		const TArray<UMovieGraphImagePassBaseNode*> ImagePassNodes =
-			FlattenedConfig->GetSettingsForBranch<UMovieGraphImagePassBaseNode>(BranchName, bIncludeCDOs, bExactMatch);
+			TemporaryEvaluatedGraph->GetSettingsForBranch<UMovieGraphImagePassBaseNode>(BranchName, bIncludeCDOs, bExactMatch);
 
 		// Update all file output nodes to have the viewport's OCIO configuration
 		for (UMovieGraphFileOutputNode* FileOutputNode : FileOutputNodes)
@@ -437,6 +476,54 @@ void FMovieGraphQuickRender::ApplyQuickRenderUpdatesToDuplicateGraph_ApplyOcio(U
 	}
 }
 
+void FMovieGraphQuickRender::ApplyQuickRenderUpdatesToDuplicateGraph_Scalability()
+{
+	constexpr bool bIncludeCDOs = false;
+	constexpr bool bExactMatch = true;
+	const UMovieGraphGlobalGameOverridesNode* GameOverridesNode =
+		TemporaryEvaluatedGraph->GetSettingForBranch<UMovieGraphGlobalGameOverridesNode>(UMovieGraphNode::GlobalsPinName, bIncludeCDOs, bExactMatch);
+
+	// If the user specified an explicit scalability level to use, there's nothing to do here. Otherwise, we'll apply the value that the viewport
+	// is using.
+	if (GameOverridesNode && GameOverridesNode->bOverride_ScalabilityQualityLevel)
+	{
+		return;
+	}
+	
+	if (UMovieGraphGlobalGameOverridesNode* NewGameOverridesNode =
+			Cast<UMovieGraphGlobalGameOverridesNode>(InjectNodeIntoBranch(UMovieGraphGlobalGameOverridesNode::StaticClass(), UMovieGraphNode::GlobalsPinName)))
+	{
+		const Scalability::FQualityLevels QualityLevels = Scalability::GetQualityLevels();
+		const int32 MinQualityLevel = QualityLevels.GetMinQualityLevel();
+		const int32 QualityLevel = QualityLevels.GetSingleQualityLevel();
+
+		int32 QualityLevelToUse;
+
+		// Use the overall value set for scalability if it's not custom (-1). If that's custom, fall back to the minimum quality level. If that's
+		// custom, then default to High.
+		if (QualityLevel != -1)
+		{
+			QualityLevelToUse = QualityLevel;
+		}
+		else if (MinQualityLevel != -1)
+		{
+			UE_LOG(LogMovieRenderPipeline, Warning, TEXT("Scalability settings are not all set to the same level, which is not supported by Quick Render. Using the minimum from all settings, which is [%s]."),
+				*Scalability::GetScalabilityNameFromQualityLevel(MinQualityLevel).ToString());
+
+			QualityLevelToUse = MinQualityLevel;
+		}
+		else
+		{
+			UE_LOG(LogMovieRenderPipeline, Warning, TEXT("One or more scalability settings are set to custom, which is not supported by Quick Render. Defaulting to High scalability"));
+
+			QualityLevelToUse = 2;	// High
+		}
+
+		NewGameOverridesNode->bOverride_ScalabilityQualityLevel = true;
+		NewGameOverridesNode->ScalabilityQualityLevel = static_cast<EMovieGraphScalabilityQualityLevel>(QualityLevelToUse);
+	}
+}
+
 UMovieGraphNode* FMovieGraphQuickRender::InjectNodeIntoBranch(const TSubclassOf<UMovieGraphNode> NodeType, const FName& InBranchName)
 {
 	check(TemporaryGraph.IsValid());
@@ -521,30 +608,24 @@ void FMovieGraphQuickRender::RestorePreRenderState(const EMovieGraphQuickRenderM
 	}
 	
 	CachedPrePieData.SelectedCameras.Empty();
-
-	// If we're rendering with the Viewport Camera in Current Sequence mode, remove the subsequence track that was added
-	if (QuickRenderMode == EMovieGraphQuickRenderMode::UseViewportCameraInSequence)
+	
+	// Remove the Utility subsequence that was added to the Rendering sequence
+	UMovieSceneTrack* TrackToRemove = nullptr;
+	for (UMovieSceneSection* Section : RenderingLevelSequence->GetMovieScene()->GetAllSections())
 	{
-		const ULevelSequence* CurrentLevelSequence = ULevelSequenceEditorBlueprintLibrary::GetCurrentLevelSequence();
-		check(CurrentLevelSequence);
-
-		UMovieSceneTrack* TrackToRemove = nullptr;
-		for (UMovieSceneSection* Section : CurrentLevelSequence->GetMovieScene()->GetAllSections())
+		if (const UMovieSceneSubSection* SubsequenceSection = Cast<UMovieSceneSubSection>(Section))
 		{
-			if (const UMovieSceneSubSection* SubsequenceSection = Cast<UMovieSceneSubSection>(Section))
+			if (SubsequenceSection->GetSequence() == UtilityLevelSequence.Get())
 			{
-				if (SubsequenceSection->GetSequence() == UtilityLevelSequence.Get())
-				{
-					TrackToRemove = Section->GetTypedOuter<UMovieSceneTrack>();
-					break;
-				}
+				TrackToRemove = Section->GetTypedOuter<UMovieSceneTrack>();
+				break;
 			}
 		}
+	}
 
-		if (TrackToRemove)
-		{
-			CurrentLevelSequence->GetMovieScene()->RemoveTrack(*TrackToRemove);
-		}
+	if (TrackToRemove)
+	{
+		RenderingLevelSequence->GetMovieScene()->RemoveTrack(*TrackToRemove);
 	}
 }
 
@@ -588,17 +669,12 @@ void FMovieGraphQuickRender::SetUpRenderingLevelSequence(const EMovieGraphQuickR
 	RenderingLevelSequence.Reset(CurrentLevelSequence);
 }
 
-void FMovieGraphQuickRender::SetUpUtilityLevelSequence(const EMovieGraphQuickRenderMode QuickRenderMode)
+void FMovieGraphQuickRender::SetUpUtilityLevelSequence()
 {
-	// This mode needs a utility sequence that contains a camera cut with a high hierarchical bias, which will allow it to override all cameras in
-	// the level sequence being rendered.
-	if (QuickRenderMode == EMovieGraphQuickRenderMode::UseViewportCameraInSequence)
-	{
-		ULevelSequence* UtilitySequence = NewObject<ULevelSequence>(GetTransientPackage(), TEXT("QuickRenderUtilitySequence"), RF_Transient);
-		UtilitySequence->Initialize();
+	ULevelSequence* UtilitySequence = NewObject<ULevelSequence>(GetTransientPackage(), TEXT("QuickRenderUtilitySequence"), RF_Transient);
+	UtilitySequence->Initialize();
 
-		UtilityLevelSequence.Reset(UtilitySequence);
-	}
+	UtilityLevelSequence.Reset(UtilitySequence);
 }
 
 UWorld* FMovieGraphQuickRender::GetWorldOfType(EWorldType::Type WorldType)
@@ -651,13 +727,15 @@ TRange<FFrameNumber> FMovieGraphQuickRender::GetPlaybackRange(const UMovieGraphQ
 	return TRange<FFrameNumber>::Inclusive(StartFrame, EndFrame);
 }
 
-void FMovieGraphQuickRender::PerformPreRenderSetup(const EMovieGraphQuickRenderMode QuickRenderMode)
+void FMovieGraphQuickRender::PerformPreRenderSetup(const EMovieGraphQuickRenderMode QuickRenderMode, UWorld* InEditorWorld, const UMovieGraphQuickRenderModeSettings& InQuickRenderSettings)
 {
 	// Adds the specified camera to the Camera Cut track in the given level sequence.
 	// If there's already a camera cut section within the Camera Cut track, a new one will be added immediately after.
 	// Each section is only one frame long.
 	auto AddCameraToLevelSequence = [](ACameraActor* InCameraToAdd, const ULevelSequence* InDestLevelSequence, const TRange<FFrameNumber> OptionalFrameRangeOverride = TRange<FFrameNumber>::Empty())
 	{
+		FMovieGraphTransactionBlocker TransactionBlocker;
+		
 		UMovieScene* MovieScene = InDestLevelSequence->GetMovieScene();
 
 		// Add a Camera Cut track if one hasn't already been added
@@ -685,15 +763,8 @@ void FMovieGraphQuickRender::PerformPreRenderSetup(const EMovieGraphQuickRenderM
 		CreateBindingParams.bAllowCustomBinding = true;
 		CreateBindingParams.BindingNameOverride = InCameraToAdd->GetActorLabel();
 
-		// Create the binding for the camera. This is wrapped in a dummy transaction so we can cancel it later on to prevent the binding creation from
-		// making an entry on the undo stack.
-#if WITH_EDITOR
-		const int32 TransactionId = GEditor->BeginTransaction(FText::GetEmpty());
-#endif
+		// Create the binding for the camera
 		const FGuid CameraBinding = FSequencerUtilities::CreateOrReplaceBinding(nullptr, MovieScene->GetTypedOuter<UMovieSceneSequence>(), InCameraToAdd, CreateBindingParams);
-#if WITH_EDITOR
-		GEditor->CancelTransaction(TransactionId);
-#endif
 
 		const TRange<FFrameNumber> SectionRange = bHasFrameRangeOverride
 			? OptionalFrameRangeOverride
@@ -714,9 +785,8 @@ void FMovieGraphQuickRender::PerformPreRenderSetup(const EMovieGraphQuickRenderM
 		// Add a new camera that mimics the viewport's camera to the level sequence. Use a regular camera actor here rather than a cine camera. Using
 		// a cine camera massively complicates things and creates a situation where it's difficult to match what the viewport looks like (eg, exposure).
 		FActorSpawnParameters CameraSpawnParams;
-		CameraSpawnParams.OverridePackage = GetTransientPackage();
-		CameraSpawnParams.bCreateActorPackage = false;
-		CameraSpawnParams.OverrideLevel = PieWorld->GetCurrentLevel();
+		CameraSpawnParams.ObjectFlags |= RF_Transient;
+		CameraSpawnParams.SpawnCollisionHandlingOverride = ESpawnActorCollisionHandlingMethod::AlwaysSpawn;
 		ACameraActor* NewCamera = Cast<ACameraActor>(PieWorld->SpawnActor<ACameraActor>(CameraSpawnParams));
 
 		// Update the camera to look like the viewport's camera
@@ -806,29 +876,79 @@ void FMovieGraphQuickRender::PerformPreRenderSetup(const EMovieGraphQuickRenderM
 		return NewCamera;
 	};
 	
+	// Adds a Level Visibility track to the Utility subsequence to ensure that the visibility state of levels in the editor is reflected
+	// in the render/PIE. Blueprint streamable levels, for example, may be shown in the editor, but will not typically show up in the
+	// render unless made visible by code or in Sequencer.
+	auto AddLevelVisibilityTracks = [InEditorWorld]() -> void
+	{
+		TArray<FName> VisibleLevelNames;
+		TArray<FName> HiddenLevelNames;
+
+		// Determine which levels are currently visible/hidden
+		for (const ULevelStreaming* Level : InEditorWorld->GetStreamingLevels())
+		{
+			if (!Level)
+			{
+				continue;
+			}
+	
+			const FName LevelName = FPackageName::GetShortFName(Level->GetWorldAssetPackageFName());
+			if (FLevelUtils::IsStreamingLevelVisibleInEditor(Level))
+			{
+				VisibleLevelNames.Add(LevelName);
+			}
+			else
+			{
+				HiddenLevelNames.Add(LevelName);
+			}
+		}
+
+		UMovieScene* UtilityMovieScene = UtilityLevelSequence->GetMovieScene();
+
+		// Add a track to show visible levels
+		UMovieSceneLevelVisibilityTrack* LevelVisibilityTrack_Visible = CastChecked<UMovieSceneLevelVisibilityTrack>(UtilityMovieScene->AddTrack(UMovieSceneLevelVisibilityTrack::StaticClass()));
+		UMovieSceneLevelVisibilitySection* LevelVisibilitySection_Visible = CastChecked<UMovieSceneLevelVisibilitySection>(LevelVisibilityTrack_Visible->CreateNewSection());
+		LevelVisibilityTrack_Visible->AddSection(*LevelVisibilitySection_Visible);
+		LevelVisibilitySection_Visible->SetVisibility(ELevelVisibility::Visible);
+		LevelVisibilitySection_Visible->SetRange(UtilityMovieScene->GetPlaybackRange());
+		LevelVisibilitySection_Visible->SetLevelNames(VisibleLevelNames);
+
+		// Add a track to hide hidden levels
+		UMovieSceneLevelVisibilityTrack* LevelVisibilityTrack_Hidden = CastChecked<UMovieSceneLevelVisibilityTrack>(UtilityMovieScene->AddTrack(UMovieSceneLevelVisibilityTrack::StaticClass()));
+		UMovieSceneLevelVisibilitySection* LevelVisibilitySection_Hidden = CastChecked<UMovieSceneLevelVisibilitySection>(LevelVisibilityTrack_Hidden->CreateNewSection());
+		LevelVisibilityTrack_Hidden->AddSection(*LevelVisibilitySection_Hidden);
+		LevelVisibilitySection_Hidden->SetVisibility(ELevelVisibility::Hidden);
+		LevelVisibilitySection_Hidden->SetRange(UtilityMovieScene->GetPlaybackRange());
+		LevelVisibilitySection_Hidden->SetLevelNames(HiddenLevelNames);
+	};
+
+	UMovieScene* RenderingMovieScene = RenderingLevelSequence->GetMovieScene();
+
+	// Always add the Utility sequence. It's used for various things in different modes. Set the hierarchical bias of the new subsequence section so
+	// its effects take precedence over the level sequence it's being added to. Use a very high number; the value used here is arbitrary, but should
+	// be high enough to do the job. Note that the subsequence section is added with 0 duration, its length will be determined later.
+	UMovieSceneSubTrack* NewSubsequenceTrack = CastChecked<UMovieSceneSubTrack>(RenderingMovieScene->AddTrack(UMovieSceneSubTrack::StaticClass()));
+	UMovieSceneSubSection* SequenceSubsection = NewSubsequenceTrack->AddSequence(UtilityLevelSequence.Get(), 0, 0);
+	SequenceSubsection->Parameters.HierarchicalBias = 1000000;
+
+	// Track the frame range that should be used for rendering. By default this is the rendering sequence's playback range, but for many modes which
+	// generate a dynamic level sequence for rendering (and the rendering sequence starts as empty), this will need to be adjusted.
+	TRange<FFrameNumber> RenderingFrameRange = RenderingMovieScene->GetPlaybackRange();
+
+	// Do mode-specific setup
 	if (QuickRenderMode == EMovieGraphQuickRenderMode::CurrentViewport)
 	{
+		// Adding a camera will tweak the length of the utility sequence's playback range
 		ACameraActor* NewCamera = DuplicateViewportCamera();
-		AddCameraToLevelSequence(NewCamera, RenderingLevelSequence.Get());
+		AddCameraToLevelSequence(NewCamera, UtilityLevelSequence.Get());
+
+		RenderingFrameRange = UtilityLevelSequence->GetMovieScene()->GetPlaybackRange();
 	}
 	else if (QuickRenderMode == EMovieGraphQuickRenderMode::UseViewportCameraInSequence)
 	{
-		UMovieScene* CurrentMovieScene = RenderingLevelSequence.Get()->GetMovieScene();
-		const TRange<FFrameNumber> CurrentPlaybackRange = CurrentMovieScene->GetPlaybackRange();
-
-		// Add a subsequence to the currently-active level sequence
-		const FFrameNumber StartFrame = CurrentPlaybackRange.GetLowerBoundValue();
-		const FFrameNumber EndFrame = CurrentPlaybackRange.GetUpperBoundValue();
-		UMovieSceneSubTrack* NewSubsequenceTrack = CastChecked<UMovieSceneSubTrack>(CurrentMovieScene->AddTrack(UMovieSceneSubTrack::StaticClass()));
-		UMovieSceneSubSection* SequenceSubsection = NewSubsequenceTrack->AddSequence(UtilityLevelSequence.Get(), StartFrame, (EndFrame - StartFrame).Value);
-
-		// Set the hierarchical bias of the new subsequence section so its camera takes precedence over all other cameras (eg, those from shots).
-		// Use a very high number; the value used here is arbitrary, but should be high enough to do the job.
-		SequenceSubsection->Parameters.HierarchicalBias = 1000000;
-
-		// Add the viewport camera to the subsequence. This is done through the level sequence that was added to it.
+		// Add the viewport camera to the utility sequence to override any other cameras in use
 		ACameraActor* NewCamera = DuplicateViewportCamera();
-		AddCameraToLevelSequence(NewCamera, UtilityLevelSequence.Get(), CurrentPlaybackRange);
+		AddCameraToLevelSequence(NewCamera, UtilityLevelSequence.Get(), RenderingMovieScene->GetPlaybackRange());
 	}
 	else if (QuickRenderMode == EMovieGraphQuickRenderMode::SelectedCameras)
 	{
@@ -836,10 +956,25 @@ void FMovieGraphQuickRender::PerformPreRenderSetup(const EMovieGraphQuickRenderM
 		{
 			if (SelectedCamera.IsValid())
 			{
-				AddCameraToLevelSequence(SelectedCamera.Get(), RenderingLevelSequence.Get());
+				// Adding a camera will tweak the length of the utility sequence's playback range
+				AddCameraToLevelSequence(SelectedCamera.Get(), UtilityLevelSequence.Get());
+
+				RenderingFrameRange = UtilityLevelSequence->GetMovieScene()->GetPlaybackRange();
 			}
 		}
 	}
+
+	// Now that cameras have been added and the length of the render has been established, make sure the playback ranges are correct
+	RenderingLevelSequence->GetMovieScene()->SetPlaybackRange(RenderingFrameRange);
+	UtilityLevelSequence->GetMovieScene()->SetPlaybackRange(RenderingFrameRange);
+	SequenceSubsection->SetRange(RenderingFrameRange);
+
+	// Add the Level Visibility tracks to the Utility sequence if matching the editor visibility was requested
+	const EMovieGraphQuickRenderViewportLookFlags LookFlags = static_cast<EMovieGraphQuickRenderViewportLookFlags>(InQuickRenderSettings.ViewportLookFlags);
+	if (EnumHasAnyFlags(LookFlags, EMovieGraphQuickRenderViewportLookFlags::Visibility))
+	{
+		AddLevelVisibilityTracks();
+	}
 }
 
 #undef LOCTEXT_NAMESPACE
diff --git a/Engine/Plugins/MovieScene/MovieRenderPipeline/Source/MovieRenderPipelineEditor/Private/MoviePipelinePIEExecutor.cpp b/Engine/Plugins/MovieScene/MovieRenderPipeline/Source/MovieRenderPipelineEditor/Private/MoviePipelinePIEExecutor.cpp
--- a/Engine/Plugins/MovieScene/MovieRenderPipeline/Source/MovieRenderPipelineEditor/Private/MoviePipelinePIEExecutor.cpp
+++ b/Engine/Plugins/MovieScene/MovieRenderPipeline/Source/MovieRenderPipelineEditor/Private/MoviePipelinePIEExecutor.cpp
@@ -72,6 +72,7 @@ UMoviePipelinePIEExecutor::UMoviePipelinePIEExecutor()
 	, RemainingInitializationFrames(-1)
 	, bPreviousUseFixedTimeStep(false)
 	, PreviousFixedTimeStepDelta(1 / 30.0)
+	, bAllowUnsavedLevels(false)
 {
 	if (!IsTemplate() && FSlateApplication::IsInitialized())
 	{
@@ -82,6 +83,13 @@ UMoviePipelinePIEExecutor::UMoviePipelinePIEExecutor()
 	}
 }
 
+void UMoviePipelinePIEExecutor::SetAllowUsingUnsavedLevels(const bool bShouldAllow)
+{
+	UE_LOG(LogMovieRenderPipeline, Log, TEXT("Updating the PIE Executor to allow using unsaved levels."));
+	
+	bAllowUnsavedLevels = bShouldAllow;
+}
+
 void UMoviePipelinePIEExecutor::Start(const UMoviePipelineExecutorJob* InJob)
 {
 	Super::Start(InJob);
@@ -103,8 +111,9 @@ void UMoviePipelinePIEExecutor::Start(const UMoviePipelineExecutorJob* InJob)
 	// Check for unsaved maps. It's pretty rare that someone actually wants to execute on an unsaved map,
 	// and it catches the much more common case of adding the job to an unsaved map and then trying to render
 	// from a newly loaded map, PIE startup will fail because the map is no longer valid.
+	// Note that opting-in to allowing unsaved levels is permitted (Quick Render is a good example of when this is OK).
 	const bool bAllMapsValid = UMoviePipelineEditorBlueprintLibrary::IsMapValidForRemoteRender(Queue->GetJobs());
-	if (!bAllMapsValid)
+	if (!bAllMapsValid && !bAllowUnsavedLevels)
 	{
 		if(!IsRenderingOffscreen())
 		{
diff --git a/Engine/Plugins/MovieScene/MovieRenderPipeline/Source/MovieRenderPipelineEditor/Public/Graph/MovieGraphQuickRender.h b/Engine/Plugins/MovieScene/MovieRenderPipeline/Source/MovieRenderPipelineEditor/Public/Graph/MovieGraphQuickRender.h
--- a/Engine/Plugins/MovieScene/MovieRenderPipeline/Source/MovieRenderPipelineEditor/Public/Graph/MovieGraphQuickRender.h
+++ b/Engine/Plugins/MovieScene/MovieRenderPipeline/Source/MovieRenderPipelineEditor/Public/Graph/MovieGraphQuickRender.h
@@ -7,6 +7,7 @@
 
 #include "Graph/MovieGraphQuickRenderSettings.h"
 #include "LevelSequence.h"
+#include "MoviePipelinePIEExecutor.h"
 #include "MoviePipelineQueue.h"
 
 class FEditorViewportClient;
@@ -49,7 +50,10 @@ private:
 	static void ApplyQuickRenderUpdatesToDuplicateGraph_ApplyEditorVisibility(const UWorld* InEditorWorld);
 
 	/** Helper for ApplyQuickRenderUpdatesToDuplicateGraph() that applies viewport OCIO (by injecting overrides to file output and renderer nodes into the graph). */
-	static void ApplyQuickRenderUpdatesToDuplicateGraph_ApplyOcio(UMoviePipelineExecutorJob* InJob, const FOpenColorIODisplayConfiguration* InOcioConfiguration);
+	static void ApplyQuickRenderUpdatesToDuplicateGraph_ApplyOcio(const FOpenColorIODisplayConfiguration* InOcioConfiguration);
+
+	/** Helper for ApplyQuickRenderUpdatesToDuplicateGraph() that applies viewport scalability if the graph didn't specify an explicit scalability setting. */
+	static void ApplyQuickRenderUpdatesToDuplicateGraph_Scalability();
 
 	/** Injects a new node in the furthest-downstream position within the specified branch in the duplicated graph. */
 	static UMovieGraphNode* InjectNodeIntoBranch(const TSubclassOf<UMovieGraphNode> NodeType, const FName& InBranchName);
@@ -69,8 +73,11 @@ private:
 	 */
 	static void SetUpRenderingLevelSequence(const EMovieGraphQuickRenderMode QuickRenderMode);
 
-	/** Populates UtilityLevelSequence with a level sequence used for utility purposes, or does nothing if no utility sequence is needed for the current mode. */
-	static void SetUpUtilityLevelSequence(const EMovieGraphQuickRenderMode QuickRenderMode);
+	/**
+	 * Populates UtilityLevelSequence with a level sequence used for utility purposes (like overriding the camera in use, or changing level
+	 * visibility). It will be injected into the rendering level sequence as a subsequence.
+	 */
+	static void SetUpUtilityLevelSequence();
 
 	/** Gets the first world found of the specified type. */
 	static UWorld* GetWorldOfType(EWorldType::Type WorldType);
@@ -82,15 +89,21 @@ private:
 	static TRange<FFrameNumber> GetPlaybackRange(const UMovieGraphQuickRenderModeSettings& InQuickRenderSettings);
 
 	/** Does any setup needed in the PIE world prior to the render starting. */
-	static void PerformPreRenderSetup(const EMovieGraphQuickRenderMode QuickRenderMode);
+	static void PerformPreRenderSetup(const EMovieGraphQuickRenderMode QuickRenderMode, UWorld* InEditorWorld, const UMovieGraphQuickRenderModeSettings& InQuickRenderSettings);
 
 private:
 	/** The temporary queue that is used by quick render. */
 	inline static TStrongObjectPtr<UMoviePipelineQueue> TemporaryQueue = nullptr;
 
+	/** The temporary executor that drives the quick render. */
+	inline static TStrongObjectPtr<UMoviePipelinePIEExecutor> TemporaryExecutor = nullptr;
+
 	/** The temporary graph that is used by quick render. Usually a duplicate of a graph created outside of quick render, but will be modified by quick render. */
 	inline static TStrongObjectPtr<UMovieGraphConfig> TemporaryGraph = nullptr;
 
+	/** The temporary evaluated graph that is generated before a render starts; some setup processes need to inspect it. */
+	inline static TStrongObjectPtr<UMovieGraphEvaluatedConfig> TemporaryEvaluatedGraph = nullptr;
+
 	/**
 	 * The level sequence that is used for rendering (ie, provided directly to MRG to render from). Sometimes this will be the current level sequence
 	 * in Sequencer, and other times it may be a temporary level sequence solely for use in Quick Render. It depends on the mode that is being used.
diff --git a/Engine/Plugins/MovieScene/MovieRenderPipeline/Source/MovieRenderPipelineEditor/Public/MoviePipelinePIEExecutor.h b/Engine/Plugins/MovieScene/MovieRenderPipeline/Source/MovieRenderPipelineEditor/Public/MoviePipelinePIEExecutor.h
--- a/Engine/Plugins/MovieScene/MovieRenderPipeline/Source/MovieRenderPipelineEditor/Public/MoviePipelinePIEExecutor.h
+++ b/Engine/Plugins/MovieScene/MovieRenderPipeline/Source/MovieRenderPipelineEditor/Public/MoviePipelinePIEExecutor.h
@@ -68,6 +68,13 @@ public:
 	UFUNCTION(BlueprintPure, Category = "Movie Render Pipeline")
 	bool IsRenderingOffscreen() const { return bRenderOffscreen; }
 
+	/**
+	 * Set whether the executor should be allowed to use unsaved levels when rendering (by default, this is not permitted). There are some very
+	 * specific circumstances where this may be OK. Do not set this unless you're certain it will not cause issues in your use case.
+	 */
+	UFUNCTION(BlueprintCallable, Category = "Movie Render Pipeline")
+	void SetAllowUsingUnsavedLevels(const bool bShouldAllow);
+
 protected:
 	virtual void Start(const UMoviePipelineExecutorJob* InJob) override;
 
@@ -106,6 +113,9 @@ private:
 	TOptional<FDateTime> CustomInitializationTime;
 	FMoviePipelineOutputData CachedOutputDataParams;
 
+	/** Whether unsaved levels can be rendered from. */
+	bool bAllowUnsavedLevels;
+
 	/** Deprecated. use OnIndividualJobWorkFinishedDelegate instead. */
 	UE_DEPRECATED(4.27, "Use OnIndividualJobWorkFinishedDelegate instead.")
 	UPROPERTY(BlueprintAssignable, Category = "Movie Render Pipeline")
