MoviePipeline: Fixed a crash that could occur with video output nodes.

If a video node was not in the Globals branch, there were multiple output branches, and a video node was not included on each output branch, a crash could happen.

diff --git a/Engine/Plugins/MovieScene/MovieRenderPipeline/Source/MovieRenderPipelineCore/Private/Graph/Nodes/MovieGraphVideoOutputNode.cpp b/Engine/Plugins/MovieScene/MovieRenderPipeline/Source/MovieRenderPipelineCore/Private/Graph/Nodes/MovieGraphVideoOutputNode.cpp
--- a/Engine/Plugins/MovieScene/MovieRenderPipeline/Source/MovieRenderPipelineCore/Private/Graph/Nodes/MovieGraphVideoOutputNode.cpp
+++ b/Engine/Plugins/MovieScene/MovieRenderPipeline/Source/MovieRenderPipelineCore/Private/Graph/Nodes/MovieGraphVideoOutputNode.cpp
@@ -76,6 +76,12 @@ void UMovieGraphVideoOutputNode::OnReceiveImageDataImpl(UMovieGraphPipeline* InP
 			continue;
 		}
 
+		// Skip this layer if it's not part of the mask provided (ie, this node isn't part of the branch the layer is on).
+		if (!InMask.Contains(RenderPassData.Key))
+		{
+			continue;
+		}
+
 		const FMovieGraphCodecWriterWithPromise* OutputWriter = GetOrCreateOutputWriter(InPipeline, InRawFrameData, RenderPassData, CompositedPasses);
 		if (!OutputWriter)
 		{
