MoviePipeline: Fixed an issue where creating a "float" variable in the graph creates a double instead.

This is actually a fairly nuanced issue. Blueprints force doubles now, even though they're shown as floats. MRG uses many Blueprint utilities and widgets, so MRG was making doubles appear as floats. The issue was solved by making floats and doubles interoperable within the graph, and clarifying in tooltips the type of float that is being used (either single- or double-precision). The pin colors for floats and doubles were made to be the same so users are not confused.

"Float" variables are still created as doubles under the hood, but they can now be connected to any float-typed pins (which are used all throughout MRG).

diff --git a/Engine/Plugins/MovieScene/MovieRenderPipeline/Source/MovieRenderPipelineCore/Private/Graph/MovieGraphConfig.cpp b/Engine/Plugins/MovieScene/MovieRenderPipeline/Source/MovieRenderPipelineCore/Private/Graph/MovieGraphConfig.cpp
--- a/Engine/Plugins/MovieScene/MovieRenderPipeline/Source/MovieRenderPipelineCore/Private/Graph/MovieGraphConfig.cpp
+++ b/Engine/Plugins/MovieScene/MovieRenderPipeline/Source/MovieRenderPipelineCore/Private/Graph/MovieGraphConfig.cpp
@@ -1721,8 +1721,8 @@ void UMovieGraphConfig::CopyOverriddenProperties(UMovieGraphNode* FromNode, UMov
 						break;
 					}
 
-					// For the connected value to be used, the type must match and the node the value is originating from must be enabled
-					if ((ConnectedValuePin->Properties.Type == InputPin->Properties.Type) && ConnectedValuePin->Node && !ConnectedValuePin->Node->IsDisabled())
+					// For the connected value to be used, the types must be compatible and the node the value is originating from must be enabled
+					if (InputPin->IsTypeCompatibleWith(ConnectedValuePin) && ConnectedValuePin->Node && !ConnectedValuePin->Node->IsDisabled())
 					{
 						ConnectionPath.Add(ConnectedValuePin);
 					}
diff --git a/Engine/Plugins/MovieScene/MovieRenderPipeline/Source/MovieRenderPipelineCore/Private/Graph/MovieGraphPin.cpp b/Engine/Plugins/MovieScene/MovieRenderPipeline/Source/MovieRenderPipelineCore/Private/Graph/MovieGraphPin.cpp
--- a/Engine/Plugins/MovieScene/MovieRenderPipeline/Source/MovieRenderPipelineCore/Private/Graph/MovieGraphPin.cpp
+++ b/Engine/Plugins/MovieScene/MovieRenderPipeline/Source/MovieRenderPipelineCore/Private/Graph/MovieGraphPin.cpp
@@ -134,7 +134,7 @@ FPinConnectionResponse UMovieGraphPin::CanCreateConnection_PinConnectionResponse
 	const bool bBothPinsAreSameType = bPinIsWildcard ||					// Any connection can be made to a wildcard pin or
 		bBothPinsAreBranch ||											// Both are branches or
 		(!Properties.bIsBranch && !InOtherPin->Properties.bIsBranch &&	// Neither is branch and
-		Properties.Type == InOtherPin->Properties.Type &&				// They have the same property type and
+		IsTypeCompatibleWith(InOtherPin) &&								// They have compatible types and
 		Properties.TypeObject == InOtherPin->Properties.TypeObject);	// They have the same type object (for enums, structs, objects, classes)
 
 	// Pins need to be the same type
@@ -433,6 +433,15 @@ bool UMovieGraphPin::IsPinDirectionCompatibleWith(const UMovieGraphPin* OtherPin
 	return IsInputPin() != OtherPin->IsInputPin();
 }
 
+bool UMovieGraphPin::IsTypeCompatibleWith(const UMovieGraphPin* InOtherPin) const
+{
+	// There's one exception to the pin-types-must-match rule. Float and double are compatible with each other.
+	const bool bFirstPinIsFloatOrDouble = (Properties.Type == EMovieGraphValueType::Float) || (Properties.Type == EMovieGraphValueType::Double);
+	const bool bSecondPinIsFloatOrDouble = (InOtherPin->Properties.Type == EMovieGraphValueType::Float) || (InOtherPin->Properties.Type == EMovieGraphValueType::Double);
+
+	return (Properties.Type == InOtherPin->Properties.Type) || (bFirstPinIsFloatOrDouble && bSecondPinIsFloatOrDouble);
+}
+
 void UMovieGraphPin::PropagatePinType(const UMovieGraphPin* InInputPin, const UMovieGraphPin* InOutputPin)
 {
 	// If the pin's nodes aren't valid, don't proceed. This shouldn't normally happen, but could theoretically occur if one of the nodes could not
diff --git a/Engine/Plugins/MovieScene/MovieRenderPipeline/Source/MovieRenderPipelineCore/Public/Graph/MovieGraphPin.h b/Engine/Plugins/MovieScene/MovieRenderPipeline/Source/MovieRenderPipelineCore/Public/Graph/MovieGraphPin.h
--- a/Engine/Plugins/MovieScene/MovieRenderPipeline/Source/MovieRenderPipelineCore/Public/Graph/MovieGraphPin.h
+++ b/Engine/Plugins/MovieScene/MovieRenderPipeline/Source/MovieRenderPipelineCore/Public/Graph/MovieGraphPin.h
@@ -161,9 +161,17 @@ public:
 	UFUNCTION(BlueprintCallable, Category = "Movie Graph")
 	bool IsConnectionToBranchAllowed(const UMovieGraphPin* OtherPin, FText& OutError) const;
 
+	/** Determines if this pin's direction is compatible with another pin's direction. One needs to be an input, and the other an output. */
 	UFUNCTION(BlueprintCallable, Category = "Movie Graph")
 	bool IsPinDirectionCompatibleWith(const UMovieGraphPin* OtherPin) const;
 
+	/**
+	 * Determines if this pin's type is compatible with another pin's type. Generally the types have to match exactly in order to be compatible, but
+	 * there may be cases where different types can be used interchangeably.
+	 */
+	UFUNCTION(BlueprintCallable, Category = "Movie Graph")
+	bool IsTypeCompatibleWith(const UMovieGraphPin* InOtherPin) const;
+
 public:
 	// The node that this pin belongs to.
 	UPROPERTY(BlueprintReadOnly, Category = "Properties")
diff --git a/Engine/Plugins/MovieScene/MovieRenderPipeline/Source/MovieRenderPipelineEditor/Private/Graph/MovieEdGraphNode.cpp b/Engine/Plugins/MovieScene/MovieRenderPipeline/Source/MovieRenderPipelineEditor/Private/Graph/MovieEdGraphNode.cpp
--- a/Engine/Plugins/MovieScene/MovieRenderPipeline/Source/MovieRenderPipelineEditor/Private/Graph/MovieEdGraphNode.cpp
+++ b/Engine/Plugins/MovieScene/MovieRenderPipeline/Source/MovieRenderPipelineEditor/Private/Graph/MovieEdGraphNode.cpp
@@ -513,6 +513,16 @@ FString UMoviePipelineEdGraphNodeBase::GetPinTooltip(const UMovieGraphPin* InPin
 	{
 		TypeText = LOCTEXT("PinTypeTooltip_Wildcard", "Any");
 	}
+	else if (InPin->Properties.Type == EMovieGraphValueType::Float)
+	{
+		// Floats and doubles are compatible with each other in MRG, so to the user, make them both appear as "float" (and give the extra hint about
+		// precision for the people who really care)
+		TypeText = LOCTEXT("PinTypeTooltip_Float", "Float (single-precision)");
+	}
+	else if (InPin->Properties.Type == EMovieGraphValueType::Double)
+	{
+		TypeText = LOCTEXT("PinTypeTooltip_Double", "Float (double-precision)");
+	}
 	else
 	{
 		TypeText = StaticEnum<EMovieGraphValueType>()->GetDisplayNameTextByValue(static_cast<int64>(PinType));
diff --git a/Engine/Plugins/MovieScene/MovieRenderPipeline/Source/MovieRenderPipelineEditor/Private/Graph/MovieGraphSchema.cpp b/Engine/Plugins/MovieScene/MovieRenderPipeline/Source/MovieRenderPipelineEditor/Private/Graph/MovieGraphSchema.cpp
--- a/Engine/Plugins/MovieScene/MovieRenderPipeline/Source/MovieRenderPipelineEditor/Private/Graph/MovieGraphSchema.cpp
+++ b/Engine/Plugins/MovieScene/MovieRenderPipeline/Source/MovieRenderPipelineEditor/Private/Graph/MovieGraphSchema.cpp
@@ -516,9 +516,11 @@ FLinearColor UMovieGraphSchema::GetTypeColor(const FName& InPinCategory, const F
 
 	if (InPinCategory == PC_Double)
 	{
-		return Settings->DoublePinTypeColor;
+		return Settings->FloatPinTypeColor;
 	}
 
+	// Use the same pin color for floats and doubles. These types can be used interchangeably within the graph, and it's confusing to have them
+	// be different colors (because it implies that they cannot be used together).
 	if (InPinCategory == PC_Real)
 	{
 		if (InPinSubCategory == PC_Float)
@@ -528,7 +530,7 @@ FLinearColor UMovieGraphSchema::GetTypeColor(const FName& InPinCategory, const F
 
 		if (InPinSubCategory == PC_Double)
 		{
-			return Settings->DoublePinTypeColor;
+			return Settings->FloatPinTypeColor;
 		}
 	}
 
