MoviePipeline: Fixed UI state issue in MRG (when marking a member as a branch or not a branch, the other properties in the details panel may not properly refresh their enable state).

Also fixed a selection state issue in the Members panel. If multiple members have the same name, when updating one of them, the wrong one may be re-selected after the update (like changing the branch state mentioned above).

diff --git a/Engine/Plugins/MovieScene/MovieRenderPipeline/Source/MovieRenderPipelineEditor/Private/Customizations/Graph/MovieGraphMemberCustomization.h b/Engine/Plugins/MovieScene/MovieRenderPipeline/Source/MovieRenderPipelineEditor/Private/Customizations/Graph/MovieGraphMemberCustomization.h
--- a/Engine/Plugins/MovieScene/MovieRenderPipeline/Source/MovieRenderPipelineEditor/Private/Customizations/Graph/MovieGraphMemberCustomization.h
+++ b/Engine/Plugins/MovieScene/MovieRenderPipeline/Source/MovieRenderPipelineEditor/Private/Customizations/Graph/MovieGraphMemberCustomization.h
@@ -135,14 +135,21 @@ protected:
 			// Note: The graph members inherit the Value property from a base class, so the enable/disable state cannot
 			// be driven by UPROPERTY metadata. Hence why this needs to be done w/ a details customization.
 			
-			// Enable/disable the value property for inputs/outputs based on whether it is specified as a branch or not
 			const UMovieGraphInterfaceBase* InterfaceBase = Cast<UMovieGraphInterfaceBase>(CustomizedObject);
+			TWeakObjectPtr WeakInterfaceBase(InterfaceBase);
+			
+			// Enable/disable the value property for inputs/outputs based on whether it is specified as a branch or not
 			if (InterfaceBase)
 			{
 				const TSharedRef<IPropertyHandle> ValueProperty = DetailBuilder.GetProperty("Value", UMovieGraphValueContainer::StaticClass());
 				if (ValueProperty->IsValidHandle())
 				{
-					DetailBuilder.EditDefaultProperty(ValueProperty)->IsEnabled(!InterfaceBase->bIsBranch);
+					const TAttribute<bool> IsEnabledAttribute = TAttribute<bool>::CreateLambda([WeakInterfaceBase]()
+					{
+						return WeakInterfaceBase.IsValid() ? !WeakInterfaceBase->bIsBranch : false;
+					});
+					
+					DetailBuilder.EditDefaultProperty(ValueProperty)->IsEnabled(IsEnabledAttribute);
 				}
 			}
 
@@ -214,18 +221,30 @@ protected:
 					.Font(IDetailLayoutBuilder::GetDetailFont())
 				];
 			}
+			
+			TWeakObjectPtr<UMovieGraphMember> WeakMemberObject = MakeWeakObjectPtr(MemberObject);
 
 			// If this is an interface (eg, input/output) only enable the type selector if it's not a branch and editable.
 			// Otherwise, editability is the only factor in enable state.
-			const bool bIsTypeRowEnabled = InterfaceBase ? (!InterfaceBase->bIsBranch && MemberObject->IsEditable()) : MemberObject->IsEditable();
+			TAttribute<bool> IsEnabledAttribute = TAttribute<bool>::CreateLambda([WeakInterfaceBase, WeakMemberObject]()
+			{
+				if (WeakInterfaceBase.IsValid() && WeakMemberObject.IsValid())
+				{
+					return !WeakInterfaceBase->bIsBranch && WeakMemberObject->IsEditable();
+				}
+				if (WeakMemberObject.IsValid())
+				{
+					return WeakMemberObject->IsEditable();
+				}
 
-			TWeakObjectPtr<UMovieGraphMember> WeakMemberObject = MakeWeakObjectPtr(MemberObject);
+				return false;
+			});
 
 			// Add a PinTypeSelector widget to pick the data type the member uses
 			IDetailCategoryBuilder& ValueCategory = DetailBuilder.EditCategory("Value");
 			ValueCategory.AddCustomRow(FText::GetEmpty())
 			.FilterString(TypeLabel)
-			.IsEnabled(bIsTypeRowEnabled)
+			.IsEnabled(IsEnabledAttribute)
 			.NameContent()
 			[
 				SNew(STextBlock)
diff --git a/Engine/Plugins/MovieScene/MovieRenderPipeline/Source/MovieRenderPipelineEditor/Private/Graph/SMovieGraphMembersTabContent.cpp b/Engine/Plugins/MovieScene/MovieRenderPipeline/Source/MovieRenderPipelineEditor/Private/Graph/SMovieGraphMembersTabContent.cpp
--- a/Engine/Plugins/MovieScene/MovieRenderPipeline/Source/MovieRenderPipelineEditor/Private/Graph/SMovieGraphMembersTabContent.cpp
+++ b/Engine/Plugins/MovieScene/MovieRenderPipeline/Source/MovieRenderPipelineEditor/Private/Graph/SMovieGraphMembersTabContent.cpp
@@ -518,7 +518,18 @@ void SMovieGraphMembersTabContent::RefreshMemberActions(UMovieGraphMember* Updat
 	// automatically based on action name, but the member may have been renamed (so an explicit re-select is needed).
 	if (SelectedMember && (SelectedMember == UpdatedMember))
 	{
-		ActionMenu->SelectItemByName(FName(SelectedMember->GetMemberName()));
+		// Members can potentially have the same name (in different sections), so when re-selecting, make sure the right section is used.
+		int32 SectionId = static_cast<int32>(EActionSection::Variables);
+		if (SelectedMember->IsA<UMovieGraphInput>())
+		{
+			SectionId = static_cast<int32>(EActionSection::Inputs);
+		}
+		else if (SelectedMember->IsA<UMovieGraphOutput>())
+		{
+			SectionId = static_cast<int32>(EActionSection::Outputs);
+		}
+		
+		ActionMenu->SelectItemByName(FName(SelectedMember->GetMemberName()), ESelectInfo::Direct, SectionId);
 	}
 }
 
