MoviePipeline: Fixed a couple small bugs within Movie Render Graph.

* Fixed an issue that could cause the OnShotFinished callback on the Execute Script node to not contain all of files generated by MRG.
* Fixed a few cases where job variables would not properly refresh when changing/clearing graphs on the job/shot.
* Updated UMovieGraphConfig::GetSettingForBranch/GetSettingsForBranch/GetSettingForTag/GetSettingsForTag() to return UMovieGraphNode* instead of UMovieGraphSettingNode*, and also renamed them to GetNodeForBranch/GetNodesForBranch/GetNodeForTag/GetNodesForTag() to make their functionality more apparent.

diff --git a/Engine/Plugins/MovieScene/MovieRenderPipeline/Source/MovieRenderPipelineCore/Private/Graph/MovieGraphConfig.cpp b/Engine/Plugins/MovieScene/MovieRenderPipeline/Source/MovieRenderPipelineCore/Private/Graph/MovieGraphConfig.cpp
--- a/Engine/Plugins/MovieScene/MovieRenderPipeline/Source/MovieRenderPipelineCore/Private/Graph/MovieGraphConfig.cpp
+++ b/Engine/Plugins/MovieScene/MovieRenderPipeline/Source/MovieRenderPipelineCore/Private/Graph/MovieGraphConfig.cpp
@@ -1302,20 +1302,20 @@ TArray<FName> UMovieGraphConfig::GetBranchNames() const
 	return BranchNames;
 }
 
-UMovieGraphSettingNode* UMovieGraphConfig::GetSettingForBranch(UClass* InClass, const FName& InBranchName, bool bExactMatch) const
+UMovieGraphNode* UMovieGraphConfig::GetNodeForBranch(UClass* InClass, const FName& InBranchName, bool bExactMatch) const
 {
-	TArray<UMovieGraphSettingNode*> AllSettings = GetSettingsForBranch(InClass, InBranchName, bExactMatch);
-	if (AllSettings.Num() > 0)
+	TArray<UMovieGraphNode*> ResultNodes = GetNodesForBranch(InClass, InBranchName, bExactMatch);
+	if (ResultNodes.Num() > 0)
 	{
-		return AllSettings[0];
+		return ResultNodes[0];
 	}
 
 	return nullptr;
 }
 
-TArray<UMovieGraphSettingNode*> UMovieGraphConfig::GetSettingsForBranch(UClass* InClass, const FName& InBranchName, bool bExactMatch) const
+TArray<UMovieGraphNode*> UMovieGraphConfig::GetNodesForBranch(UClass* InClass, const FName& InBranchName, bool bExactMatch) const
 {
-	TArray<UMovieGraphSettingNode*> ResultNodes;
+	TArray<UMovieGraphNode*> ResultNodes;
 
 	// We need to be able to detect when we've reached the end of traversal for the current branch,
 	// which will either be at the input node, or when there's nothing connected to the left of it.
@@ -1352,7 +1352,6 @@ TArray<UMovieGraphSettingNode*> UMovieGraphConfig::GetSettingsForBranch(UClass*
 
 			if (bIsTraversingBranch)
 			{
-
 				// If there's nothing to our left on this node (or we're the Input node) then we've fully explored this branch
 				if (VisitedNode->GetInputPins().Num() == 0 || VisitedNode == LocalInputNode)
 				{
@@ -1364,7 +1363,7 @@ TArray<UMovieGraphSettingNode*> UMovieGraphConfig::GetSettingsForBranch(UClass*
 					bool bMatches = bExactMatch ? VisitedNode->GetClass() == InClass : VisitedNode->IsA(InClass);
 					if (bMatches)
 					{
-						ResultNodes.AddUnique(CastChecked<UMovieGraphSettingNode>(const_cast<UMovieGraphNode*>(VisitedNode)));
+						ResultNodes.AddUnique(const_cast<UMovieGraphNode*>(VisitedNode));
 					}
 				}
 			}
@@ -1376,56 +1375,52 @@ TArray<UMovieGraphSettingNode*> UMovieGraphConfig::GetSettingsForBranch(UClass*
 	return ResultNodes;
 }
 
-UMovieGraphSettingNode* UMovieGraphConfig::GetSettingForTag(const FString& ScriptTag, UClass* OptionalClassFilter, const FName OptionalBranchFilter, bool bExactMatch) const
+UMovieGraphNode* UMovieGraphConfig::GetNodeForTag(const FString& ScriptTag, UClass* OptionalClassFilter, const FName OptionalBranchFilter, bool bExactMatch) const
 {
-	TArray<UMovieGraphSettingNode*> AllSettings = GetSettingsForTag(ScriptTag, OptionalClassFilter, OptionalBranchFilter, bExactMatch);
-	if (AllSettings.Num() > 0)
+	TArray<UMovieGraphNode*> ResultNodes = GetNodesForTag(ScriptTag, OptionalClassFilter, OptionalBranchFilter, bExactMatch);
+	if (ResultNodes.Num() > 0)
 	{
-		return AllSettings[0];
+		return ResultNodes[0];
 	}
 
 	return nullptr;
 }
 
-TArray<UMovieGraphSettingNode*> UMovieGraphConfig::GetSettingsForTag(const FString& ScriptTag, UClass* OptionalClassFilter, const FName OptionalBranchFilter, bool bExactMatch) const
+TArray<UMovieGraphNode*> UMovieGraphConfig::GetNodesForTag(const FString& ScriptTag, UClass* OptionalClassFilter, const FName OptionalBranchFilter, bool bExactMatch) const
 {
-	TArray<UMovieGraphSettingNode*> AllSettingsToConsider;
-	if(OptionalBranchFilter != NAME_None)
+	TArray<UMovieGraphNode*> AllNodesToConsider;
+	if (OptionalBranchFilter != NAME_None)
 	{
-		UClass* ClassFilter = OptionalClassFilter ? OptionalClassFilter : UMovieGraphSettingNode::StaticClass();
-		AllSettingsToConsider = GetSettingsForBranch(ClassFilter, OptionalBranchFilter, bExactMatch);
+		UClass* ClassFilter = OptionalClassFilter ? OptionalClassFilter : UMovieGraphNode::StaticClass();
+		AllNodesToConsider = GetNodesForBranch(ClassFilter, OptionalBranchFilter, bExactMatch);
 	}
 	else
 	{
-		AllSettingsToConsider.Reserve(GetNodes().Num());
-		for(const TObjectPtr<UMovieGraphNode>& Node : GetNodes())
+		AllNodesToConsider.Reserve(GetNodes().Num());
+		for (const TObjectPtr<UMovieGraphNode>& Node : GetNodes())
 		{
-			if (UMovieGraphSettingNode* AsSetting = Cast<UMovieGraphSettingNode>(Node))
-			{
-				AllSettingsToConsider.Add(AsSetting);
-			}
+			AllNodesToConsider.Add(Node.Get());
 		}
 	}
-			
-	
-	TArray<UMovieGraphSettingNode*> AllSettingsWithTag;
-	for (UMovieGraphSettingNode* SettingNode : AllSettingsToConsider)
+
+	TArray<UMovieGraphNode*> AllNodesWithTag;
+	for (UMovieGraphNode* GraphNode : AllNodesToConsider)
 	{
 		// We re-run the class filter here because if they weren't using branch filters, they didn't get class-filtered yet.
 		const bool bHasClassFilter = OptionalClassFilter != nullptr;
-		const bool bPassesClassFilter = bExactMatch ? SettingNode->GetClass() == OptionalClassFilter : SettingNode->IsA(OptionalClassFilter);
+		const bool bPassesClassFilter = bExactMatch ? GraphNode->GetClass() == OptionalClassFilter : GraphNode->IsA(OptionalClassFilter);
 
 		const bool bConsider = bHasClassFilter ? bPassesClassFilter : true;
 		if (bConsider)
 		{
-			if (SettingNode->ScriptTags.Contains(ScriptTag))
+			if (GraphNode->ScriptTags.Contains(ScriptTag))
 			{
-				AllSettingsWithTag.Add(SettingNode);
+				AllNodesWithTag.Add(GraphNode);
 			}
 		}
 	}
 
-	return AllSettingsWithTag;
+	return AllNodesWithTag;
 }
 
 void UMovieGraphConfig::GetOutputDirectory(FString& OutOutputDirectory) const
diff --git a/Engine/Plugins/MovieScene/MovieRenderPipeline/Source/MovieRenderPipelineCore/Private/Graph/MovieGraphPipeline.cpp b/Engine/Plugins/MovieScene/MovieRenderPipeline/Source/MovieRenderPipelineCore/Private/Graph/MovieGraphPipeline.cpp
--- a/Engine/Plugins/MovieScene/MovieRenderPipeline/Source/MovieRenderPipelineCore/Private/Graph/MovieGraphPipeline.cpp
+++ b/Engine/Plugins/MovieScene/MovieRenderPipeline/Source/MovieRenderPipelineCore/Private/Graph/MovieGraphPipeline.cpp
@@ -896,27 +896,26 @@ void UMovieGraphPipeline::TeardownShot(const TObjectPtr<UMoviePipelineExecutorSh
 		GetDataSourceInstance()->RestoreHierarchyForShot(Shot);
 	}
 
-	// some other stuff
-
 	const FMovieGraphTimeStepData& TimeStepData = GetTimeStepInstance()->GetCalculatedTimeData();
 	TObjectPtr<UMovieGraphEvaluatedConfig> EvaluatedConfig = TimeStepData.EvaluatedConfig;
 
 	ProcessOutstandingFinishedFrames();
 
+	const bool bIsPostShotCallbackNeeded = IsPostShotCallbackNeeded();
+
 	// Run any post-render file generation that the nodes need to do
-	bool bIncludeCDOs = false;
-	bool bExactMatch = false;
-	const TArray<UMovieGraphFileOutputNode*> FileOutputNodes =
-		EvaluatedConfig->GetSettingsForBranch<UMovieGraphFileOutputNode>(UMovieGraphNode::GlobalsPinName, bIncludeCDOs, bExactMatch);
-	for (UMovieGraphFileOutputNode* FileOutputNode : FileOutputNodes)
+	for (const TObjectPtr<UMovieGraphFileOutputNode>& FileOutputNode : GetOutputNodesUsed())
 	{
-		FileOutputNode->OnAllShotFramesSubmitted(this, InShot, EvaluatedConfig);
+		const bool bFlushToDisk = bIsPostShotCallbackNeeded;
+		FileOutputNode->OnAllShotFramesSubmitted(this, InShot, EvaluatedConfig, bFlushToDisk);
 	}
 
-	// Ensure all of our Futures have been converted to the GeneratedOutputData
-	ProcessOutstandingFutures();
+	// Ensure all of our Futures have been converted to the GeneratedOutputData. If we have a post-shot callback, then all futures need to be completed
+	// at this point in the pipeline.
+	const bool bStallAndProcessAllFutures = bIsPostShotCallbackNeeded;
+	ProcessOutstandingFutures(bStallAndProcessAllFutures);
 
-	// Run any shot-based exports now that file output nodes have had a chance to run
+	// Run any shot-based exports now that file output nodes have had a chance to run. This includes nodes like the Command Line Encoder node.
 	const TArray<IMovieGraphPostRenderNode*> PostRenderNodes =
 		EvaluatedConfig->GetSettingsImplementing<IMovieGraphPostRenderNode>(UMovieGraphPostRenderNode::StaticClass(), UMovieGraphNode::GlobalsPinName);
 	for (IMovieGraphPostRenderNode* PostRenderNode : PostRenderNodes)
@@ -953,8 +952,8 @@ void UMovieGraphPipeline::TeardownShot(const TObjectPtr<UMoviePipelineExecutorSh
 
 	// Revert cvars set by the global game overrides. Needs to be done after the CVarManager reverts (since the global
 	// game overrides are applied first in SetupShot).
-	bIncludeCDOs = false;
-	bExactMatch = true;
+	constexpr bool bIncludeCDOs = false;
+	constexpr bool bExactMatch = true;
 	if (UMovieGraphGlobalGameOverridesNode* GlobalGameOverridesNode = EvaluatedConfig->GetSettingForBranch<UMovieGraphGlobalGameOverridesNode>(UMovieGraphNode::GlobalsPinName, bIncludeCDOs, bExactMatch))
 	{
 		constexpr bool bOverrideValues = false;
@@ -964,13 +963,12 @@ void UMovieGraphPipeline::TeardownShot(const TObjectPtr<UMoviePipelineExecutorSh
 	constexpr bool bIsGraph = true;
 	FMoviePipelineTelemetry::SendEndShotRenderTelemetry(bIsGraph, !bShutdownSetErrorFlag, bShutdownRequested);
 
-	if (IsPostShotCallbackNeeded())
+	if (bIsPostShotCallbackNeeded)
 	{
+		// Even though these methods were called above, there may be a point in the future where post-render nodes (like the Command Line Encoder
+		// node) generate files, so we need to ensure that their futures have a chance to update the generated file list.
 		ProcessOutstandingFinishedFrames();
-		ProcessOutstandingFutures();
-
-		// ToDo: Allow the command line encoder to modify the file list
-		// since it may add or remove files...
+		ProcessOutstandingFutures(bStallAndProcessAllFutures);
 
 		FMoviePipelineOutputData Params;
 		Params.Pipeline = this;
@@ -1576,19 +1574,36 @@ void UMovieGraphPipeline::AddOutputFuture(TFuture<bool>&& InOutputFuture, const
 	);
 }
 
-void UMovieGraphPipeline::ProcessOutstandingFutures()
+void UMovieGraphPipeline::ProcessOutstandingFutures(const bool bStallAndProcessAllFutures)
 {
 	// Check if any frames failed to output
 	TArray<int32> CompletedOutputFutures;
 	for (int32 Index = 0; Index < OutstandingOutputFutures.Num(); ++Index)
 	{
+		const UE::MovieGraph::FMovieGraphOutputFuture& OutputFuture = OutstandingOutputFutures[Index];
+		
 		// Output futures are pushed in order to the OutputFutures array. However they are
-		// completed asyncronously, so we don't process any futures after a not-yet-ready one
+		// completed asynchronously, so we don't process any futures after a not-yet-ready one
 		// otherwise we push into the GeneratedShotOutputData array out of order.
-		const UE::MovieGraph::FMovieGraphOutputFuture& OutputFuture = OutstandingOutputFutures[Index];
-		if (!OutputFuture.Get<0>().IsReady())
+		//
+		// Optionally, we can wait for the future being processed to actually finish before proceeding
+		// to the next. This is helpful when we need to guarantee that all futures are processed before
+		// moving on to the next step in the pipeline (eg, a step that depends on images being written
+		// to disk).
+		if (bStallAndProcessAllFutures)
 		{
-			break;
+			while (!OutputFuture.Get<0>().IsReady())
+			{
+				// Sleep for a while to allow the futures to complete
+				FPlatformProcess::Sleep(0.1f);
+			}
+		}
+		else
+		{
+			if (!OutputFuture.Get<0>().IsReady())
+			{
+				break;
+			}
 		}
 
 		CompletedOutputFutures.Add(Index);
diff --git a/Engine/Plugins/MovieScene/MovieRenderPipeline/Source/MovieRenderPipelineCore/Private/Graph/Nodes/MovieGraphAudioOutputNode.cpp b/Engine/Plugins/MovieScene/MovieRenderPipeline/Source/MovieRenderPipelineCore/Private/Graph/Nodes/MovieGraphAudioOutputNode.cpp
--- a/Engine/Plugins/MovieScene/MovieRenderPipeline/Source/MovieRenderPipelineCore/Private/Graph/Nodes/MovieGraphAudioOutputNode.cpp
+++ b/Engine/Plugins/MovieScene/MovieRenderPipeline/Source/MovieRenderPipelineCore/Private/Graph/Nodes/MovieGraphAudioOutputNode.cpp
@@ -91,7 +91,7 @@ void UMovieGraphAudioOutputNode::OnAllFramesSubmittedImpl(UMovieGraphPipeline* I
 	StartAudioExport();
 }
 
-void UMovieGraphAudioOutputNode::OnAllShotFramesSubmittedImpl(UMovieGraphPipeline* InPipeline, const UMoviePipelineExecutorShot* InShot, TObjectPtr<UMovieGraphEvaluatedConfig>& InShotEvaluatedGraph)
+void UMovieGraphAudioOutputNode::OnAllShotFramesSubmittedImpl(UMovieGraphPipeline* InPipeline, const UMoviePipelineExecutorShot* InShot, TObjectPtr<UMovieGraphEvaluatedConfig>& InShotEvaluatedGraph, const bool bFlushToDisk)
 {
 	CachedPipeline = InPipeline;
 	EvaluatedGraph = InPipeline->GetTimeStepInstance()->GetCalculatedTimeData().EvaluatedConfig;
diff --git a/Engine/Plugins/MovieScene/MovieRenderPipeline/Source/MovieRenderPipelineCore/Private/Graph/Nodes/MovieGraphVideoOutputNode.cpp b/Engine/Plugins/MovieScene/MovieRenderPipeline/Source/MovieRenderPipelineCore/Private/Graph/Nodes/MovieGraphVideoOutputNode.cpp
--- a/Engine/Plugins/MovieScene/MovieRenderPipeline/Source/MovieRenderPipelineCore/Private/Graph/Nodes/MovieGraphVideoOutputNode.cpp
+++ b/Engine/Plugins/MovieScene/MovieRenderPipeline/Source/MovieRenderPipelineCore/Private/Graph/Nodes/MovieGraphVideoOutputNode.cpp
@@ -21,14 +21,14 @@ UMovieGraphVideoOutputNode::UMovieGraphVideoOutputNode()
 	FileNameFormat = TEXT("{sequence_name}.{layer_name}");
 }
 
-void UMovieGraphVideoOutputNode::OnAllShotFramesSubmittedImpl(UMovieGraphPipeline* InPipeline, const UMoviePipelineExecutorShot* InShot, TObjectPtr<UMovieGraphEvaluatedConfig>& InShotEvaluatedGraph)
+void UMovieGraphVideoOutputNode::OnAllShotFramesSubmittedImpl(UMovieGraphPipeline* InPipeline, const UMoviePipelineExecutorShot* InShot, TObjectPtr<UMovieGraphEvaluatedConfig>& InShotEvaluatedGraph, const bool bFlushToDisk)
 {
 	constexpr bool bIncludeCDOs = true;
 	constexpr bool bExactMatch = true;
 	const UMovieGraphGlobalOutputSettingNode* OutputSettingsNode =
 		InShotEvaluatedGraph->GetSettingForBranch<UMovieGraphGlobalOutputSettingNode>(GlobalsPinName, bIncludeCDOs, bExactMatch);
 	
-	if (OutputSettingsNode->bFlushDiskWritesPerShot)
+	if (bFlushToDisk)
 	{
 		// If they don't have {shot_name} or {camera_name} in their output path, this probably doesn't do what they expect
 		// as it will finalize and then overwrite itself.
@@ -210,10 +210,7 @@ void UMovieGraphVideoOutputNode::OnAllFramesFinalizedImpl(UMovieGraphPipeline* I
 
 		Finalize_EncodeThread(RawWriter);
 
-		if (!bHasError)
-		{
-			Writer.Promise.SetValue(true);
-		}
+		Writer.Promise.SetValue(!bHasError);
 	}
 
 	// AllWriters is static and shared across *all* video output node types. Only delete writers of the current node's type.
diff --git a/Engine/Plugins/MovieScene/MovieRenderPipeline/Source/MovieRenderPipelineCore/Private/MoviePipelineQueue.cpp b/Engine/Plugins/MovieScene/MovieRenderPipeline/Source/MovieRenderPipelineCore/Private/MoviePipelineQueue.cpp
--- a/Engine/Plugins/MovieScene/MovieRenderPipeline/Source/MovieRenderPipelineCore/Private/MoviePipelineQueue.cpp
+++ b/Engine/Plugins/MovieScene/MovieRenderPipeline/Source/MovieRenderPipelineCore/Private/MoviePipelineQueue.cpp
@@ -492,10 +492,8 @@ void UMoviePipelineExecutorShot::RefreshAllVariableAssignments()
 	// Refresh the associated primary graph variable override assignments as well
 	if (const UMoviePipelineExecutorJob* PrimaryJob = GetTypedOuter<UMoviePipelineExecutorJob>())
 	{
-		if (UMovieGraphConfig* PrimaryGraph = PrimaryJob->GetGraphPreset())
-		{
-			MoviePipeline::RefreshVariableAssignments(PrimaryGraph, PrimaryGraphVariableAssignments, this);
-		}
+		UMovieGraphConfig* PrimaryGraph = PrimaryJob->GetGraphPreset();
+		MoviePipeline::RefreshVariableAssignments(PrimaryGraph, PrimaryGraphVariableAssignments, this);
 	}
 }
 
diff --git a/Engine/Plugins/MovieScene/MovieRenderPipeline/Source/MovieRenderPipelineCore/Private/MoviePipelineUtils.cpp b/Engine/Plugins/MovieScene/MovieRenderPipeline/Source/MovieRenderPipelineCore/Private/MoviePipelineUtils.cpp
--- a/Engine/Plugins/MovieScene/MovieRenderPipeline/Source/MovieRenderPipelineCore/Private/MoviePipelineUtils.cpp
+++ b/Engine/Plugins/MovieScene/MovieRenderPipeline/Source/MovieRenderPipelineCore/Private/MoviePipelineUtils.cpp
@@ -265,14 +265,14 @@ namespace MoviePipeline
 
 	void RefreshVariableAssignments(UMovieGraphConfig* InRootGraph, TArray<TObjectPtr<UMovieJobVariableAssignmentContainer>>& InVariableAssignments, UObject* InAssignmentsOwner)
 	{
-		if (!InRootGraph)
+		TSet<UMovieGraphConfig*> AllGraphs;
+		
+		if (InRootGraph)
 		{
-			return;
+			AllGraphs = {InRootGraph};
+			InRootGraph->GetAllContainedSubgraphs(AllGraphs);
 		}
 	
-		TSet<UMovieGraphConfig*> AllGraphs = {InRootGraph};
-		InRootGraph->GetAllContainedSubgraphs(AllGraphs);
-
 		// Add/update variable assignments for the graph on the job and all of its subgraphs
 		for (const UMovieGraphConfig* Graph : AllGraphs)
 		{
diff --git a/Engine/Plugins/MovieScene/MovieRenderPipeline/Source/MovieRenderPipelineCore/Public/Graph/MovieGraphConfig.h b/Engine/Plugins/MovieScene/MovieRenderPipeline/Source/MovieRenderPipelineCore/Public/Graph/MovieGraphConfig.h
--- a/Engine/Plugins/MovieScene/MovieRenderPipeline/Source/MovieRenderPipelineCore/Public/Graph/MovieGraphConfig.h
+++ b/Engine/Plugins/MovieScene/MovieRenderPipeline/Source/MovieRenderPipelineCore/Public/Graph/MovieGraphConfig.h
@@ -624,12 +624,12 @@ public:
 	UFUNCTION(BlueprintCallable, Category = "Movie Graph")
 	bool RemoveLabeledEdge(UMovieGraphNode* FromNode, const FName& FromPinName, UMovieGraphNode* ToNode, const FName& ToPinName);
 	
-	
 	/**
 	* Convenience function which removes all Inbound (pins on the left side of a node) edges for the given node.
 	*/
 	UFUNCTION(BlueprintCallable, Category = "Movie Graph")
 	bool RemoveAllInboundEdges(UMovieGraphNode* InNode);
+	
 	/**
 	* Convenience function which removes all Outbound (pins on the right side of a node) edges for the given node.
 	*/
@@ -680,74 +680,76 @@ public:
 	TArray<FName> GetBranchNames() const;
 
 	/**
-	* Finds a setting node (by type) for the given branch name (see GetBranchNames()). Returns the first result
+	* Finds a node (by type) for the given branch name (see GetBranchNames()). Returns the first result
 	* of that type found, even if there are multiple, following traversal order (ie: right to left).
 	* To prevent accidentally editing unrelated assets, does not dive into sub-graphs (but can continue
 	* traversal beyond them), so returned results should only exist in current asset.
-	* Does not contain the Input or Output nodes (as they are not settings), see GetInputNode() and GetOutputNode().
+	* Does not contain the Input or Output nodes, see GetInputNode() and GetOutputNode().
 	*
-	* @param InClass - The node class type to search for.
+	* @param InClass - The node class type to search for. Use UMovieGraphSettingNode to only include nodes that show up in the final evaluated graph (excludes nodes like the Branch node).
 	* @param InBranchName - The branch to search for these nodes on.
 	* @param bExactMatch - Defaults false. If true, will not consider inherited classes of InClass.
 	* 					   (ie: If requesting a base class and only child classes exist, they will not count when using exact matching).
 	* @return First node found (if any) which are of the specified InClass type on the given branch.
 	*/
 	UFUNCTION(BlueprintCallable, meta = (DeterminesOutputType = "InClass"), Category = "Movie Graph")
-	UMovieGraphSettingNode* GetSettingForBranch(UClass* InClass, const FName& InBranchName, bool bExactMatch = false) const;
+	UMovieGraphNode* GetNodeForBranch(UClass* InClass, const FName& InBranchName, bool bExactMatch = false) const;
 
 	/**
-	* Finds all setting nodes (by type) for the given branch name (see GetBranchNames()). Returns all results of that type
+	* Finds all nodes (by type) for the given branch name (see GetBranchNames()). Returns all results of that type
 	* following traversal order (ie: right to left).
 	* To prevent accidentally editing unrelated assets, does not dive into sub-graphs (but can continue
 	* traversal beyond them), so returned results should only exist in current asset.
-	* Does not contain the Input or Output nodes (as they are not settings), see GetInputNode() and GetOutputNode().
+	* Does not contain the Input or Output nodes, see GetInputNode() and GetOutputNode().
 	*
-	* @param InClass - The node class type to search for.
+	* @param InClass - The node class type to search for. Use UMovieGraphSettingNode to only include nodes that show up in the final evaluated graph (excludes nodes like the Branch node).
 	* @param InBranchName - The branch to search for these nodes on.
 	* @param bExactMatch - Defaults false. If true, will not consider inherited classes of InClass.
 	* 					   (ie: If requesting a base class and only child classes exist, they will not count when using exact matching).
 	* @return All nodes found (if any) which are of the specified InClass type on the given branch.
 	*/
 	UFUNCTION(BlueprintCallable, meta = (DeterminesOutputType = "InClass"), Category = "Movie Graph")
-	TArray<UMovieGraphSettingNode*> GetSettingsForBranch(UClass* InClass, const FName& InBranchName, bool bExactMatch = false) const;
+	TArray<UMovieGraphNode*> GetNodesForBranch(UClass* InClass, const FName& InBranchName, bool bExactMatch = false) const;
 
 	/**
-	* Finds a setting node by ScriptTag field. THIS SEARCH IS CASE SENSITIVE. Returns the first result
+	* Finds a node by ScriptTag field. THIS SEARCH IS CASE SENSITIVE. Returns the first result
 	* of that tag found, even if there are multiple, following traversal order (ie: right to left).
 	* If OptionalClassFilter is specified, only matches against nodes that have the right class type and
 	* contain the correct tag.
 	* To prevent accidentally editing unrelated assets, does not dive into sub-graphs (but can continue
 	* traversal beyond them), so returned results should only exist in current asset.
-	* Does not contain the Input or Output nodes (as they are not settings), see GetInputNode() and GetOutputNode().
+	* Does not contain the Input or Output nodes, see GetInputNode() and GetOutputNode().
 	*
 	* @param ScriptTag - Case sensitive name into the "Script Tags" property on Nodes
-	* @param OptionalClassFilter - Defaults None. If specified, will only return a node of this class that matches the tag.
+	* @param OptionalClassFilter - Defaults None. If specified, will only return a node of this class that matches the tag. Use UMovieGraphSettingNode to only include nodes
+	*							   that show up in the final evaluated graph (excludes nodes like the Branch node).
 	* @param OptionalBranchFilter - Defaults None. If specified, will only search the specified branch, otherwise searches all branches.
 	* @param bExactMatch - Defaults false. If true, will not consider inherited classes to pass the OptionalClassFilter.
 	* 					   (ie: If requesting a base class and only child classes exist, they will not count when using exact matching).
 	* @return The first node found (if any) which contains the case-sensitive tag and is of the correct class type (if using OptionalClassFilter)
 	*/
 	UFUNCTION(BlueprintCallable, meta = (DeterminesOutputType = "OptionalClassFilter"), Category = "Movie Graph")
-	UMovieGraphSettingNode* GetSettingForTag(const FString& ScriptTag, UClass* OptionalClassFilter = nullptr, const FName OptionalBranchFilter = NAME_None, bool bExactMatch = false) const;
+	UMovieGraphNode* GetNodeForTag(const FString& ScriptTag, UClass* OptionalClassFilter = nullptr, const FName OptionalBranchFilter = NAME_None, bool bExactMatch = false) const;
 
 	/**
-	* Finds setting nodes by ScriptTag field. THIS SEARCH IS CASE SENSITIVE. Returns all results
+	* Finds nodes by ScriptTag field. THIS SEARCH IS CASE SENSITIVE. Returns all results
 	* of that tag found, following traversal order (ie: right to left).
 	* If OptionalClassFilter is specified, only matches against nodes that have the right class type and
 	* contain the correct tag.
 	* To prevent accidentally editing unrelated assets, does not dive into sub-graphs (but can continue
 	* traversal beyond them), so returned results should only exist in current asset.
-	* Does not contain the Input or Output nodes (as they are not settings), see GetInputNode(), GetOutputNode())
+	* Does not contain the Input or Output nodes, see GetInputNode(), GetOutputNode())
 	*
 	* @param ScriptTag - Case sensitive name into the "Script Tags" property on Nodes
-	* @param OptionalClassFilter - Defaults None. If specified, will only return a node of this class that matches the tag.
+	* @param OptionalClassFilter - Defaults None. If specified, will only return a node of this class that matches the tag. Use UMovieGraphSettingNode to only include nodes
+	*							   that show up in the final evaluated graph (excludes nodes like the Branch node).
 	* @param OptionalBranchFilter - Defaults None. If specified, will only search the specified branch, otherwise searches all branches.
 	* @param bExactMatch - Defaults false. If true, will not consider inherited classes to pass the OptionalClassFilter.
 	* 					   (ie: If requesting a base class and only child classes exist, they will not count when using exact matching).
 	* @return All nodes found (if any) which contain the case-sensitive tag and are of the correct class type (if using OptionalClassFilter)
 	*/
 	UFUNCTION(BlueprintCallable, meta = (DeterminesOutputType = "OptionalClassFilter"), Category = "Movie Graph")
-	TArray<UMovieGraphSettingNode*> GetSettingsForTag(const FString& ScriptTag, UClass* OptionalClassFilter = nullptr, const FName OptionalBranchFilter = NAME_None, bool bExactMatch = false) const;
+	TArray<UMovieGraphNode*> GetNodesForTag(const FString& ScriptTag, UClass* OptionalClassFilter = nullptr, const FName OptionalBranchFilter = NAME_None, bool bExactMatch = false) const;
 
 	/**
 	 * Adds a new variable member with default values to the graph. The new variable will have a base name of
diff --git a/Engine/Plugins/MovieScene/MovieRenderPipeline/Source/MovieRenderPipelineCore/Public/Graph/MovieGraphPipeline.h b/Engine/Plugins/MovieScene/MovieRenderPipeline/Source/MovieRenderPipelineCore/Public/Graph/MovieGraphPipeline.h
--- a/Engine/Plugins/MovieScene/MovieRenderPipeline/Source/MovieRenderPipelineCore/Public/Graph/MovieGraphPipeline.h
+++ b/Engine/Plugins/MovieScene/MovieRenderPipeline/Source/MovieRenderPipelineCore/Public/Graph/MovieGraphPipeline.h
@@ -150,8 +150,21 @@ protected:
 	virtual void OnEngineTickEndFrame();
 	virtual void RenderFrame();
 	virtual void BuildShotListFromDataSource();
+
+	/**
+	 * Sends the rendered pixel data to all file output nodes (they will use the "mask" provided to determine if they need to process the data).
+	 * OnReceiveImageData() will be called on each of these nodes, and in turn they will typically call AddOutputFuture() on the pipeline to
+	 * provide a future that communicates the status of their write task.
+	 */
 	virtual void ProcessOutstandingFinishedFrames();
-	virtual void ProcessOutstandingFutures();
+
+	/**
+	 * Processes any futures that have completed by updating GeneratedOutputData. Futures are processed sequentially in the order that they're
+	 * generated. If the future next-in-line to be processed has not completed, this will early exit, and must be called again to continue
+	 * processing futures. Optionally, all futures can be processed by setting "bStallAndProcessAllFutures" to true; this will still process
+	 * all futures sequentially, but waits for the current one to finish before moving on to the next (and so on).
+	 */
+	virtual void ProcessOutstandingFutures(const bool bStallAndProcessAllFutures = false);
 
 	virtual void TickProducingFrames();
 	virtual void TickPostFinalizeExport(const bool bInForceFinish);
diff --git a/Engine/Plugins/MovieScene/MovieRenderPipeline/Source/MovieRenderPipelineCore/Public/Graph/Nodes/MovieGraphAudioOutputNode.h b/Engine/Plugins/MovieScene/MovieRenderPipeline/Source/MovieRenderPipelineCore/Public/Graph/Nodes/MovieGraphAudioOutputNode.h
--- a/Engine/Plugins/MovieScene/MovieRenderPipeline/Source/MovieRenderPipelineCore/Public/Graph/Nodes/MovieGraphAudioOutputNode.h
+++ b/Engine/Plugins/MovieScene/MovieRenderPipeline/Source/MovieRenderPipelineCore/Public/Graph/Nodes/MovieGraphAudioOutputNode.h
@@ -47,7 +47,7 @@ public:
 protected:
 	// UMovieGraphFileOutputNode Interface
 	virtual void OnAllFramesSubmittedImpl(UMovieGraphPipeline* InPipeline, TObjectPtr<UMovieGraphEvaluatedConfig>& InPrimaryJobEvaluatedGraph) override;
-	virtual void OnAllShotFramesSubmittedImpl(UMovieGraphPipeline* InPipeline, const UMoviePipelineExecutorShot* InShot, TObjectPtr<UMovieGraphEvaluatedConfig>& InShotEvaluatedGraph) override;
+	virtual void OnAllShotFramesSubmittedImpl(UMovieGraphPipeline* InPipeline, const UMoviePipelineExecutorShot* InShot, TObjectPtr<UMovieGraphEvaluatedConfig>& InShotEvaluatedGraph, const bool bFlushToDisk) override;
 	virtual bool IsFinishedWritingToDiskImpl() const override;
 	// ~UMovieGraphFileOutputNode Interface
 
diff --git a/Engine/Plugins/MovieScene/MovieRenderPipeline/Source/MovieRenderPipelineCore/Public/Graph/Nodes/MovieGraphFileOutputNode.h b/Engine/Plugins/MovieScene/MovieRenderPipeline/Source/MovieRenderPipelineCore/Public/Graph/Nodes/MovieGraphFileOutputNode.h
--- a/Engine/Plugins/MovieScene/MovieRenderPipeline/Source/MovieRenderPipelineCore/Public/Graph/Nodes/MovieGraphFileOutputNode.h
+++ b/Engine/Plugins/MovieScene/MovieRenderPipeline/Source/MovieRenderPipelineCore/Public/Graph/Nodes/MovieGraphFileOutputNode.h
@@ -39,10 +39,11 @@ public:
 	void OnAllFramesFinalized(UMovieGraphPipeline* InPipeline, TObjectPtr<UMovieGraphEvaluatedConfig>& InPrimaryJobEvaluatedGraph) { OnAllFramesFinalizedImpl(InPipeline, InPrimaryJobEvaluatedGraph); }
 
 	/**
-	 * This is called when a shot has finished rendering its last frame, and the shot export process has not begun yet.
+	 * This is called when a shot has finished rendering its last frame, and the shot export process has not begun yet. If bFlushToDisk is true,
+	 * all files this node is generating for the current shot should be written to disk before returning.
 	 * See OnAllFramesSubmitted() for a method that is called after all shots have finished rendering.
 	 */
-	void OnAllShotFramesSubmitted(UMovieGraphPipeline* InPipeline, const UMoviePipelineExecutorShot* InShot, TObjectPtr<UMovieGraphEvaluatedConfig>& InShotEvaluatedGraph) { OnAllShotFramesSubmittedImpl(InPipeline, InShot, InShotEvaluatedGraph); }
+	void OnAllShotFramesSubmitted(UMovieGraphPipeline* InPipeline, const UMoviePipelineExecutorShot* InShot, TObjectPtr<UMovieGraphEvaluatedConfig>& InShotEvaluatedGraph, const bool bFlushToDisk = false) { OnAllShotFramesSubmittedImpl(InPipeline, InShot, InShotEvaluatedGraph, bFlushToDisk); }
 
 	/** Returns whether this node has finished writing all of its files to disk yet. */
 	bool IsFinishedWritingToDisk() const { return IsFinishedWritingToDiskImpl(); }
@@ -81,7 +82,11 @@ protected:
 	virtual void OnReceiveImageDataImpl(UMovieGraphPipeline* InPipeline, UE::MovieGraph::FMovieGraphOutputMergerFrame* InRawFrameData, const TSet<FMovieGraphRenderDataIdentifier>& InMask) {}
 	virtual void OnAllFramesSubmittedImpl(UMovieGraphPipeline* InPipeline, TObjectPtr<UMovieGraphEvaluatedConfig>& InPrimaryJobEvaluatedGraph) {}
 	virtual void OnAllFramesFinalizedImpl(UMovieGraphPipeline* InPipeline, TObjectPtr<UMovieGraphEvaluatedConfig>& InPrimaryJobEvaluatedGraph) {}
-	virtual void OnAllShotFramesSubmittedImpl(UMovieGraphPipeline* InPipeline, const UMoviePipelineExecutorShot* InShot, TObjectPtr<UMovieGraphEvaluatedConfig>& InShotEvaluatedGraph) {}
+	virtual void OnAllShotFramesSubmittedImpl(UMovieGraphPipeline* InPipeline, const UMoviePipelineExecutorShot* InShot, TObjectPtr<UMovieGraphEvaluatedConfig>& InShotEvaluatedGraph, const bool bFlushToDisk) {}
+
+	UE_DEPRECATED(5.6, "Please use the version of OnAllShotFramesSubmittedImpl() that has the new bFlushToDisk argument at the end.")
+	virtual void OnAllShotFramesSubmittedImpl(UMovieGraphPipeline* InPipeline, const UMoviePipelineExecutorShot* InShot, TObjectPtr<UMovieGraphEvaluatedConfig>& InShotEvaluatedGraph) final {}
+	
 	virtual bool IsFinishedWritingToDiskImpl() const { return true; }
 
 	/** Returns whether any overscan should be automatically cropped from the image before compositing and output */
diff --git a/Engine/Plugins/MovieScene/MovieRenderPipeline/Source/MovieRenderPipelineCore/Public/Graph/Nodes/MovieGraphVideoOutputNode.h b/Engine/Plugins/MovieScene/MovieRenderPipeline/Source/MovieRenderPipelineCore/Public/Graph/Nodes/MovieGraphVideoOutputNode.h
--- a/Engine/Plugins/MovieScene/MovieRenderPipeline/Source/MovieRenderPipelineCore/Public/Graph/Nodes/MovieGraphVideoOutputNode.h
+++ b/Engine/Plugins/MovieScene/MovieRenderPipeline/Source/MovieRenderPipelineCore/Public/Graph/Nodes/MovieGraphVideoOutputNode.h
@@ -61,7 +61,7 @@ protected:
 	virtual void OnReceiveImageDataImpl(UMovieGraphPipeline* InPipeline, UE::MovieGraph::FMovieGraphOutputMergerFrame* InRawFrameData, const TSet<FMovieGraphRenderDataIdentifier>& InMask) override;
 	virtual void OnAllFramesSubmittedImpl(UMovieGraphPipeline* InPipeline, TObjectPtr<UMovieGraphEvaluatedConfig>& InPrimaryJobEvaluatedGraph) override;
 	virtual void OnAllFramesFinalizedImpl(UMovieGraphPipeline* InPipeline, TObjectPtr<UMovieGraphEvaluatedConfig>& InPrimaryJobEvaluatedGraph) override;
-	virtual void OnAllShotFramesSubmittedImpl(UMovieGraphPipeline* InPipeline, const UMoviePipelineExecutorShot* InShot, TObjectPtr<UMovieGraphEvaluatedConfig>& InShotEvaluatedGraph) override;
+	virtual void OnAllShotFramesSubmittedImpl(UMovieGraphPipeline* InPipeline, const UMoviePipelineExecutorShot* InShot, TObjectPtr<UMovieGraphEvaluatedConfig>& InShotEvaluatedGraph, const bool bFlushToDisk) override;
 	virtual bool IsFinishedWritingToDiskImpl() const override;
 	// ~UMovieGraphFileOutputNode Interface
 
