MoviePipeline: Added support for controlling Post-Process Volume visibility within MRG modifiers.

diff --git a/Engine/Plugins/MovieScene/MovieRenderPipeline/Source/MovieRenderPipelineCore/Private/Graph/MovieGraphRenderLayerSubsystem.cpp b/Engine/Plugins/MovieScene/MovieRenderPipeline/Source/MovieRenderPipelineCore/Private/Graph/MovieGraphRenderLayerSubsystem.cpp
--- a/Engine/Plugins/MovieScene/MovieRenderPipeline/Source/MovieRenderPipelineCore/Private/Graph/MovieGraphRenderLayerSubsystem.cpp
+++ b/Engine/Plugins/MovieScene/MovieRenderPipeline/Source/MovieRenderPipelineCore/Private/Graph/MovieGraphRenderLayerSubsystem.cpp
@@ -247,9 +247,17 @@ void UMovieGraphRenderPropertyModifier::ApplyModifier(const UWorld* World)
 
 	auto CacheActorState = [this](AActor* InActor)
 	{
+		bool bIsActorHidden = InActor->IsHidden();
+
+		// PPV actors use the "enabled" flag for visibility
+		if (const APostProcessVolume* PpvActor = Cast<APostProcessVolume>(InActor))
+		{
+			bIsActorHidden = !PpvActor->bEnabled;
+		}
+		
 		FActorState& CachedState = CachedActorState.AddDefaulted_GetRef();
 		CachedState.Actor = InActor;
-		CachedState.bIsHidden = InActor->IsHidden();
+		CachedState.bIsHidden = bIsActorHidden;
 	};
 	
 	CachedActorState.Empty();
@@ -395,6 +403,12 @@ void UMovieGraphRenderPropertyModifier::SetActorAndComponentState(const TArray<F
 			return;
 		}
 
+		// Post-process volumes have an "enabled" flag that we'll use to hide/un-hide them
+		if (APostProcessVolume* PpvActor = Cast<APostProcessVolume>(InActor))
+		{
+			PpvActor->bEnabled = !InNewActorState.bIsHidden;
+		}
+
 		InActor->SetActorHiddenInGame(InNewActorState.bIsHidden);
 #if WITH_EDITOR
 		InActor->SetIsTemporarilyHiddenInEditor(InNewActorState.bIsHidden);
