MoviePipeline: Exposed Quick Render to Python and Blueprints.

Also:
* Added an example on how to use Quick Render with Python.
* Added some utility methods that may be helpful with scripting.
* Added the ability for scripting to override the level sequence being used in Quick Render (this is not exposed in the UI).
* Converted Quick Render over to being an editor subsystem for easier use within scripting and BPs, with the added benefit of having a better system architecture internally as well.

diff --git a/Engine/Plugins/MovieScene/MovieRenderPipeline/Source/MovieRenderPipelineCore/Private/Graph/MovieGraphConfig.cpp b/Engine/Plugins/MovieScene/MovieRenderPipeline/Source/MovieRenderPipelineCore/Private/Graph/MovieGraphConfig.cpp
--- a/Engine/Plugins/MovieScene/MovieRenderPipeline/Source/MovieRenderPipelineCore/Private/Graph/MovieGraphConfig.cpp
+++ b/Engine/Plugins/MovieScene/MovieRenderPipeline/Source/MovieRenderPipelineCore/Private/Graph/MovieGraphConfig.cpp
@@ -915,6 +915,20 @@ UMovieGraphVariable* UMovieGraphConfig::GetVariableByGuid(const FGuid& InGuid) c
 	return nullptr;
 }
 
+UMovieGraphVariable* UMovieGraphConfig::GetVariableByName(const FString& InVariableName) const
+{
+	constexpr bool bIncludeGlobal = true;
+	for (UMovieGraphVariable* Variable : GetVariables(bIncludeGlobal))
+	{
+		if (Variable->GetMemberName() == InVariableName)
+		{
+			return Variable;
+		}
+	}
+
+	return nullptr;
+}
+
 TArray<UMovieGraphVariable*> UMovieGraphConfig::GetVariables(const bool bIncludeGlobal) const
 {
 	if (!bIncludeGlobal)
diff --git a/Engine/Plugins/MovieScene/MovieRenderPipeline/Source/MovieRenderPipelineCore/Private/Graph/MovieGraphQuickRenderSettings.cpp b/Engine/Plugins/MovieScene/MovieRenderPipeline/Source/MovieRenderPipelineCore/Private/Graph/MovieGraphQuickRenderSettings.cpp
--- a/Engine/Plugins/MovieScene/MovieRenderPipeline/Source/MovieRenderPipelineCore/Private/Graph/MovieGraphQuickRenderSettings.cpp
+++ b/Engine/Plugins/MovieScene/MovieRenderPipeline/Source/MovieRenderPipelineCore/Private/Graph/MovieGraphQuickRenderSettings.cpp
@@ -111,6 +111,19 @@ void UMovieGraphQuickRenderModeSettings::RefreshVariableAssignments(UMovieGraphQ
 	MoviePipeline::RefreshVariableAssignments(InSettings->GraphPreset.LoadSynchronous(), InSettings->GraphVariableAssignments, InSettings);
 }
 
+UMovieJobVariableAssignmentContainer* UMovieGraphQuickRenderModeSettings::GetVariableAssignmentsForGraph(const TSoftObjectPtr<UMovieGraphConfig>& InGraphConfigPath) const
+{
+	for (const TObjectPtr<UMovieJobVariableAssignmentContainer>& VariableAssignments : GraphVariableAssignments)
+	{
+		if (VariableAssignments->GetGraphConfig() == InGraphConfigPath)
+		{
+			return VariableAssignments.Get();
+		}
+	}
+
+	return nullptr;
+}
+
 #if WITH_EDITOR
 void UMovieGraphQuickRenderModeSettings::PostEditChangeProperty(FPropertyChangedEvent& PropertyChangedEvent)
 {
diff --git a/Engine/Plugins/MovieScene/MovieRenderPipeline/Source/MovieRenderPipelineCore/Public/Graph/MovieGraphConfig.h b/Engine/Plugins/MovieScene/MovieRenderPipeline/Source/MovieRenderPipelineCore/Public/Graph/MovieGraphConfig.h
--- a/Engine/Plugins/MovieScene/MovieRenderPipeline/Source/MovieRenderPipelineCore/Public/Graph/MovieGraphConfig.h
+++ b/Engine/Plugins/MovieScene/MovieRenderPipeline/Source/MovieRenderPipelineCore/Public/Graph/MovieGraphConfig.h
@@ -779,6 +779,10 @@ public:
 	/** Gets the variable in the graph with the specified GUID, else nullptr if one could not be found. */
 	UFUNCTION(BlueprintCallable, Category="Movie Graph")
 	UMovieGraphVariable* GetVariableByGuid(const FGuid& InGuid) const;
+	
+	/** Gets the variable in the graph with the specified name (including global variables), else nullptr if one could not be found. */
+	UFUNCTION(BlueprintCallable, Category="Movie Graph")
+	UMovieGraphVariable* GetVariableByName(const FString& InVariableName) const;
 
 	/**
 	 * Gets all variables that are available to be used in the graph. Global variables can optionally be included if
diff --git a/Engine/Plugins/MovieScene/MovieRenderPipeline/Source/MovieRenderPipelineCore/Public/Graph/MovieGraphQuickRenderSettings.h b/Engine/Plugins/MovieScene/MovieRenderPipeline/Source/MovieRenderPipelineCore/Public/Graph/MovieGraphQuickRenderSettings.h
--- a/Engine/Plugins/MovieScene/MovieRenderPipeline/Source/MovieRenderPipelineCore/Public/Graph/MovieGraphQuickRenderSettings.h
+++ b/Engine/Plugins/MovieScene/MovieRenderPipeline/Source/MovieRenderPipelineCore/Public/Graph/MovieGraphQuickRenderSettings.h
@@ -133,6 +133,15 @@ public:
 	UFUNCTION(BlueprintCallable, Category="Quick Render")
 	static void RefreshVariableAssignments(UMovieGraphQuickRenderModeSettings* InSettings);
 
+	/**
+	 * Gets the variable assignments for the specified graph asset. Note that there will normally only be one variable assignment container for these
+	 * quick render settings (the assignments for the graph specified in `GraphPreset`). However, `GraphPreset` may have subgraphs, and subgraphs
+	 * store their variable assignments in separate containers. If you need to update the variable assignments for those subgraph(s), this method
+	 * provides a convenient way of getting their assignments (instead of iterating through `GraphVariableAssignments`).
+	 */
+	UFUNCTION(BlueprintCallable, Category="Quick Render")
+	UMovieJobVariableAssignmentContainer* GetVariableAssignmentsForGraph(const TSoftObjectPtr<UMovieGraphConfig>& InGraphConfigPath) const;
+
 #if WITH_EDITOR
 	FOnMovieGraphQuickRenderGraphChanged OnGraphChangedDelegate;
 #endif
@@ -140,6 +149,14 @@ public:
 	/** The graph preset that is used to configure the render. If not specified, the quick render default graph will be used. */
 	UPROPERTY(EditAnywhere, BlueprintReadWrite, Category="Configuration", DisplayName="Movie Render Graph Config")
 	TSoftObjectPtr<UMovieGraphConfig> GraphPreset;
+	
+	/**
+	 * For sequence-centric modes (like "Current Sequence"), if this override is set, this sequence will be used instead of the level sequence that's
+	 * active in Sequencer. In UI-based Quick Render, this property is never set directly because Sequencer is always used. However, for scripting
+	 * purposes, there may be cases where Sequencer is not open, and a specific level sequence needs to be used.
+	 */
+	UPROPERTY(BlueprintReadWrite, Category="Configuration", DisplayName="Movie Render Graph Config")
+	TSoftObjectPtr<class ULevelSequence> LevelSequenceOverride = nullptr;
 
 	/** The action that quick render should perform after a render is finished. */
 	UPROPERTY(EditAnywhere, BlueprintReadWrite, Category="Configuration", DisplayName="After Render")
@@ -160,11 +177,11 @@ public:
 	UPROPERTY(EditAnywhere, BlueprintReadWrite, Category="Quick Render", DisplayName="Sequencer Frame Range")
 	EMovieGraphQuickRenderFrameRangeType FrameRangeType = EMovieGraphQuickRenderFrameRangeType::PlaybackRange;
 
-	/** The frame in the level sequence that the render will begin on. */
+	/** The frame in the level sequence that the render will begin on. This frame number is inclusive, meaning that this frame will be included in the rendered frames. */
 	UPROPERTY(EditAnywhere, BlueprintReadWrite, Category="Quick Render", meta=(EditCondition="FrameRangeType == EMovieGraphQuickRenderFrameRangeType::Custom", EditConditionHides))
 	int32 CustomStartFrame = 0;
 
-	/** The frame in the level sequence that the render will end on. */
+	/** The frame in the level sequence that the render will end on. This frame number is exclusive, meaning that the rendered frames will include the frames up until (but excluding) this frame. */
 	UPROPERTY(EditAnywhere, BlueprintReadWrite, Category="Quick Render", meta=(EditCondition="FrameRangeType == EMovieGraphQuickRenderFrameRangeType::Custom", EditConditionHides))
 	int32 CustomEndFrame = 0;
 
diff --git a/Engine/Plugins/MovieScene/MovieRenderPipeline/Source/MovieRenderPipelineCore/Public/MovieJobVariableAssignmentContainer.h b/Engine/Plugins/MovieScene/MovieRenderPipeline/Source/MovieRenderPipelineCore/Public/MovieJobVariableAssignmentContainer.h
--- a/Engine/Plugins/MovieScene/MovieRenderPipeline/Source/MovieRenderPipelineCore/Public/MovieJobVariableAssignmentContainer.h
+++ b/Engine/Plugins/MovieScene/MovieRenderPipeline/Source/MovieRenderPipelineCore/Public/MovieJobVariableAssignmentContainer.h
@@ -32,6 +32,7 @@ public:
 	void SetGraphConfig(const TSoftObjectPtr<UMovieGraphConfig>& InGraphConfig);
 
 	/** Gets the graph that is associated with this container. */
+	UFUNCTION(BlueprintCallable, Category="Graph")
 	TSoftObjectPtr<UMovieGraphConfig> GetGraphConfig() const;
 
 	/** Gets the number of variable assignments present in this container. Assignments that are disabled are counted. */
diff --git a/Engine/Plugins/MovieScene/MovieRenderPipeline/Source/MovieRenderPipelineEditor/Private/Graph/MovieGraphQuickRender.cpp b/Engine/Plugins/MovieScene/MovieRenderPipeline/Source/MovieRenderPipelineEditor/Private/Graph/MovieGraphQuickRender.cpp
--- a/Engine/Plugins/MovieScene/MovieRenderPipeline/Source/MovieRenderPipelineEditor/Private/Graph/MovieGraphQuickRender.cpp
+++ b/Engine/Plugins/MovieScene/MovieRenderPipeline/Source/MovieRenderPipelineEditor/Private/Graph/MovieGraphQuickRender.cpp
@@ -76,7 +76,7 @@ private:
 	int32 TransactionId = 0;
 };
 
-void FMovieGraphQuickRender::BeginQuickRender(const EMovieGraphQuickRenderMode QuickRenderMode, const UMovieGraphQuickRenderModeSettings& InQuickRenderSettings)
+void UMovieGraphQuickRenderSubsystem::BeginQuickRender(const EMovieGraphQuickRenderMode QuickRenderMode, const UMovieGraphQuickRenderModeSettings* InQuickRenderSettings)
 {
 	UMoviePipelineQueueSubsystem* Subsystem = GEditor->GetEditorSubsystem<UMoviePipelineQueueSubsystem>();
 	check(Subsystem);
@@ -87,9 +87,12 @@ void FMovieGraphQuickRender::BeginQuickRender(const EMovieGraphQuickRenderMode Q
 		return;
 	}
 
+	// Store a reference to the settings used, so the executor callbacks can rely on them still existing.
+	QuickRenderModeSettings = InQuickRenderSettings;
+
 	// Get/create the level sequence that should be rendered from (varies depending on the Quick Render mode).
 	SetUpRenderingLevelSequence(QuickRenderMode);
-	if (!RenderingLevelSequence.IsValid())
+	if (!IsValid(RenderingLevelSequence))
 	{
 		return;
 	}
@@ -112,7 +115,7 @@ void FMovieGraphQuickRender::BeginQuickRender(const EMovieGraphQuickRenderMode Q
 	}
 
 	// Allocate a new queue that will be used temporarily
-	TemporaryQueue = TStrongObjectPtr(NewObject<UMoviePipelineQueue>(GetTransientPackage()));
+	TemporaryQueue = NewObject<UMoviePipelineQueue>(GetTransientPackage());
 
 	// Add a temp job to the queue, and give it the editor's level and the current level sequence being edited in Sequencer
 	UMoviePipelineExecutorJob* NewJob = TemporaryQueue->AllocateNewJob(UMoviePipelineExecutorJob::StaticClass());
@@ -122,7 +125,7 @@ void FMovieGraphQuickRender::BeginQuickRender(const EMovieGraphQuickRenderMode Q
 
 	// Update the job to use the quick render graph (either the default, or the one specified in the settings)
 	TMap<TObjectPtr<UMovieGraphConfig>, TObjectPtr<UMovieGraphConfig>> OriginalGraphToDupeMap;
-	const UMovieGraphConfig* QuickRenderGraph = GetQuickRenderGraph(InQuickRenderSettings.GraphPreset.LoadSynchronous(), OriginalGraphToDupeMap);
+	const UMovieGraphConfig* QuickRenderGraph = GetQuickRenderGraph(QuickRenderModeSettings->GraphPreset.LoadSynchronous(), OriginalGraphToDupeMap);
 	if (!QuickRenderGraph)
 	{
 		// GetQuickRenderGraph() will output an error to the log if a graph was not returned
@@ -131,7 +134,7 @@ void FMovieGraphQuickRender::BeginQuickRender(const EMovieGraphQuickRenderMode Q
 	NewJob->SetGraphPreset(QuickRenderGraph);
 
 	// Apply any job-level variable overrides, if they were specified
-	for (const TObjectPtr<UMovieJobVariableAssignmentContainer>& IncomingVariableAssignment : InQuickRenderSettings.GraphVariableAssignments)
+	for (const TObjectPtr<UMovieJobVariableAssignmentContainer>& IncomingVariableAssignment : QuickRenderModeSettings->GraphVariableAssignments)
 	{
 		// Map the setting's graph to the duplicate graph
 		const UMovieGraphConfig* AssignmentGraph = IncomingVariableAssignment->GetGraphConfig().LoadSynchronous();
@@ -162,7 +165,7 @@ void FMovieGraphQuickRender::BeginQuickRender(const EMovieGraphQuickRenderMode Q
 	FMovieGraphTraversalContext Context;
 	Context.Job = NewJob;
 	Context.RootGraph = TemporaryGraph.Get();
-	TemporaryEvaluatedGraph = TStrongObjectPtr(TemporaryGraph->CreateFlattenedGraph(Context, OutError));
+	TemporaryEvaluatedGraph = TemporaryGraph->CreateFlattenedGraph(Context, OutError);
 	if (!TemporaryEvaluatedGraph)
 	{
 		UE_LOG(LogMovieRenderPipeline, Error, TEXT("Unable to Quick Render: Could not create an evaluated graph. Reason: %s"), *OutError);
@@ -171,7 +174,7 @@ void FMovieGraphQuickRender::BeginQuickRender(const EMovieGraphQuickRenderMode Q
 
 	// Update the graph as needed in order to properly apply the settings specified by the user. This will perform a graph evaluation, so the job
 	// should be fully initialized at this point.
-	ApplyQuickRenderUpdatesToDuplicateGraph(InQuickRenderSettings, EditorWorld, NewJob);
+	ApplyQuickRenderUpdatesToDuplicateGraph(EditorWorld);
 
 	// Cache out anything necessary before PIE starts
 	CachePrePieData(EditorWorld);
@@ -179,19 +182,19 @@ void FMovieGraphQuickRender::BeginQuickRender(const EMovieGraphQuickRenderMode Q
 	// Create the executor. Note that it's allowed to use unsaved levels (many animators work off of unsaved levels). Because Quick Render is
 	// creating a temporary queue and rendering locally, using an unsaved level is safe.
 	// Always use the PIE executor for Quick Render, rather than the executor specified in the project settings for normal MRG.
-	TemporaryExecutor = TStrongObjectPtr(NewObject<UMoviePipelinePIEExecutor>(GetTransientPackage()));
+	TemporaryExecutor = NewObject<UMoviePipelinePIEExecutor>(GetTransientPackage());
 	TemporaryExecutor->SetAllowUsingUnsavedLevels(true);
 
 	// Do any needed setup in the PIE world before rendering starts
-	TemporaryExecutor->OnIndividualJobStarted().AddLambda([QuickRenderMode, EditorWorld, &InQuickRenderSettings](UMoviePipelineExecutorJob* InJob)
+	TemporaryExecutor->OnIndividualJobStarted().AddLambda([this, QuickRenderMode, EditorWorld](UMoviePipelineExecutorJob* InJob)
 	{
-		PerformPreRenderSetup(QuickRenderMode, EditorWorld, InQuickRenderSettings);
+		PerformPreRenderSetup(QuickRenderMode, EditorWorld);
 	});
 
 	// View the job's rendered frame(s)/video(s) if requested
-	TemporaryExecutor->OnIndividualJobWorkFinished().AddLambda([&InQuickRenderSettings](FMoviePipelineOutputData OutputData)
+	TemporaryExecutor->OnIndividualJobWorkFinished().AddLambda([this](FMoviePipelineOutputData OutputData)
 	{
-		if (InQuickRenderSettings.PostRenderBehavior == EMovieGraphQuickRenderPostRenderActionType::PlayRenderOutput)
+		if (QuickRenderModeSettings->PostRenderBehavior == EMovieGraphQuickRenderPostRenderActionType::PlayRenderOutput)
 		{
 			FMoviePipelinePostRenderFileDisplayProcessor PostRenderFileDisplayProcessor(GetDefault<UMovieRenderGraphEditorSettings>()->PostRenderSettings);
 			PostRenderFileDisplayProcessor.AddFiles(OutputData);
@@ -200,24 +203,25 @@ void FMovieGraphQuickRender::BeginQuickRender(const EMovieGraphQuickRenderMode Q
 	});
 
 	// Clean up
-	TemporaryExecutor->OnExecutorFinished().AddLambda([QuickRenderMode](UMoviePipelineExecutorBase* InExecutor, bool bSuccess)
+	TemporaryExecutor->OnExecutorFinished().AddLambda([this, QuickRenderMode](UMoviePipelineExecutorBase* InExecutor, bool bSuccess)
 	{
 		RestorePreRenderState(QuickRenderMode);
 
 		// The temp objects are part of the transient package, so manually release them when the render is finished
-		TemporaryQueue.Reset();
-		TemporaryGraph.Reset();
-		TemporaryExecutor.Reset();
-		TemporaryEvaluatedGraph.Reset();
-		RenderingLevelSequence.Reset();
-		UtilityLevelSequence.Reset();
+		TemporaryQueue = nullptr;
+		TemporaryGraph = nullptr;
+		TemporaryExecutor = nullptr;
+		TemporaryEvaluatedGraph = nullptr;
+		RenderingLevelSequence = nullptr;
+		UtilityLevelSequence = nullptr;
+		QuickRenderModeSettings = nullptr;
 	});
 
 	// Do the render
 	Subsystem->RenderQueueInstanceWithExecutorInstance(TemporaryQueue.Get(), TemporaryExecutor.Get());
 }
 
-UMovieGraphConfig* FMovieGraphQuickRender::GetQuickRenderGraph(UMovieGraphConfig* InUserSpecifiedGraph, TMap<TObjectPtr<UMovieGraphConfig>, TObjectPtr<UMovieGraphConfig>>& OutOriginalToDupeMap)
+UMovieGraphConfig* UMovieGraphQuickRenderSubsystem::GetQuickRenderGraph(UMovieGraphConfig* InUserSpecifiedGraph, TMap<TObjectPtr<UMovieGraphConfig>, TObjectPtr<UMovieGraphConfig>>& OutOriginalToDupeMap)
 {
 	UMovieGraphConfig* QuickRenderGraph = InUserSpecifiedGraph;
 	if (!QuickRenderGraph)
@@ -246,14 +250,14 @@ UMovieGraphConfig* FMovieGraphQuickRender::GetQuickRenderGraph(UMovieGraphConfig
 
 	// Duplicate the graph so changes can be made to it without affecting the graph asset being used
 	UMovieGraphConfig* DuplicateGraph = MoviePipeline::DuplicateConfigRecursive(QuickRenderGraph, OutOriginalToDupeMap);
-	TemporaryGraph = TStrongObjectPtr(DuplicateGraph);
+	TemporaryGraph = DuplicateGraph;
 
 	return TemporaryGraph.Get();
 }
 
-void FMovieGraphQuickRender::ApplyQuickRenderUpdatesToDuplicateGraph(const UMovieGraphQuickRenderModeSettings& InQuickRenderSettings, UWorld* InEditorWorld, UMoviePipelineExecutorJob* InJob)
+void UMovieGraphQuickRenderSubsystem::ApplyQuickRenderUpdatesToDuplicateGraph(const UWorld* InEditorWorld)
 {
-	check(TemporaryGraph.IsValid());
+	check(IsValid(TemporaryGraph));
 
 	// Apply changes to Global Output Settings
 	{
@@ -273,7 +277,7 @@ void FMovieGraphQuickRender::ApplyQuickRenderUpdatesToDuplicateGraph(const UMovi
 			GlobalOutputSettings->CustomPlaybackRangeEnd.Value = EndFrame;
 		};
 
-		const TRange<FFrameNumber> PlaybackRange = GetPlaybackRange(InQuickRenderSettings);
+		const TRange<FFrameNumber> PlaybackRange = GetPlaybackRange();
 		if (!PlaybackRange.IsEmpty())
 		{
 			SetStartFrame(PlaybackRange.GetLowerBoundValue().Value);
@@ -283,7 +287,7 @@ void FMovieGraphQuickRender::ApplyQuickRenderUpdatesToDuplicateGraph(const UMovi
 
 	// Inject an Apply Viewport Look node if any viewport look flags were specified. Logic outside of Quick Render will for look this node in order to
 	// determine how to apply some of its settings
-	const EMovieGraphQuickRenderViewportLookFlags LookFlags = static_cast<EMovieGraphQuickRenderViewportLookFlags>(InQuickRenderSettings.ViewportLookFlags);
+	const EMovieGraphQuickRenderViewportLookFlags LookFlags = static_cast<EMovieGraphQuickRenderViewportLookFlags>(QuickRenderModeSettings->ViewportLookFlags);
 	if (LookFlags != EMovieGraphQuickRenderViewportLookFlags::None)
 	{
 		UMovieGraphApplyViewportLookNode* ApplyViewportLookNode = Cast<UMovieGraphApplyViewportLookNode>(InjectNodeIntoBranch(UMovieGraphApplyViewportLookNode::StaticClass(), UMovieGraphNode::GlobalsPinName));
@@ -311,7 +315,7 @@ void FMovieGraphQuickRender::ApplyQuickRenderUpdatesToDuplicateGraph(const UMovi
 	}
 
 	// Apply OCIO if it is activated on the viewport and enabled in the viewport look flags
-	if (InQuickRenderSettings.bOverride_ViewportLookFlags && EnumHasAnyFlags(LookFlags, EMovieGraphQuickRenderViewportLookFlags::Ocio))
+	if (QuickRenderModeSettings->bOverride_ViewportLookFlags && EnumHasAnyFlags(LookFlags, EMovieGraphQuickRenderViewportLookFlags::Ocio))
 	{
 #if WITH_OCIO
 		if (const FEditorViewportClient* ViewportClient = UMovieGraphApplyViewportLookNode::GetViewportClient())
@@ -331,7 +335,7 @@ void FMovieGraphQuickRender::ApplyQuickRenderUpdatesToDuplicateGraph(const UMovi
 	ApplyQuickRenderUpdatesToDuplicateGraph_Scalability();
 }
 
-void FMovieGraphQuickRender::ApplyQuickRenderUpdatesToDuplicateGraph_ApplyEditorVisibility(const UWorld* InEditorWorld)
+void UMovieGraphQuickRenderSubsystem::ApplyQuickRenderUpdatesToDuplicateGraph_ApplyEditorVisibility(const UWorld* InEditorWorld)
 {
 	TArray<FMovieGraphActorQueryEntry> VisibleActors;
 	TArray<FMovieGraphActorQueryEntry> HiddenActors;
@@ -366,7 +370,7 @@ void FMovieGraphQuickRender::ApplyQuickRenderUpdatesToDuplicateGraph_ApplyEditor
 		}
 	}
 
-	auto AddNewCollectionWithVisibilityModifier = []<typename QueryType>(const FString& InOperationName, const bool bModifierShouldHide) -> QueryType*
+	auto AddNewCollectionWithVisibilityModifier = [this]<typename QueryType>(const FString& InOperationName, const bool bModifierShouldHide) -> QueryType*
 	{
 		const UMovieGraphCollectionNode* NewCollectionNode = Cast<UMovieGraphCollectionNode>(InjectNodeIntoBranch(UMovieGraphCollectionNode::StaticClass(), UMovieGraphNode::GlobalsPinName));
 
@@ -397,7 +401,7 @@ void FMovieGraphQuickRender::ApplyQuickRenderUpdatesToDuplicateGraph_ApplyEditor
 	HiddenActorQuery->ActorsAndComponentsToMatch = HiddenActors;
 }
 
-void FMovieGraphQuickRender::ApplyQuickRenderUpdatesToDuplicateGraph_ApplyOcio(const FOpenColorIODisplayConfiguration* InOcioConfiguration)
+void UMovieGraphQuickRenderSubsystem::ApplyQuickRenderUpdatesToDuplicateGraph_ApplyOcio(const FOpenColorIODisplayConfiguration* InOcioConfiguration)
 {
 	// For each branch, apply OCIO configuration.
 	for (const FName& BranchName : TemporaryEvaluatedGraph->GetBranchNames())
@@ -476,7 +480,7 @@ void FMovieGraphQuickRender::ApplyQuickRenderUpdatesToDuplicateGraph_ApplyOcio(c
 	}
 }
 
-void FMovieGraphQuickRender::ApplyQuickRenderUpdatesToDuplicateGraph_Scalability()
+void UMovieGraphQuickRenderSubsystem::ApplyQuickRenderUpdatesToDuplicateGraph_Scalability()
 {
 	constexpr bool bIncludeCDOs = false;
 	constexpr bool bExactMatch = true;
@@ -524,9 +528,9 @@ void FMovieGraphQuickRender::ApplyQuickRenderUpdatesToDuplicateGraph_Scalability
 	}
 }
 
-UMovieGraphNode* FMovieGraphQuickRender::InjectNodeIntoBranch(const TSubclassOf<UMovieGraphNode> NodeType, const FName& InBranchName)
+UMovieGraphNode* UMovieGraphQuickRenderSubsystem::InjectNodeIntoBranch(const TSubclassOf<UMovieGraphNode> NodeType, const FName& InBranchName) const
 {
-	check(TemporaryGraph.IsValid());
+	check(IsValid(TemporaryGraph));
 
 	UMovieGraphNode* NewNode = TemporaryGraph->CreateNodeByClass(NodeType);
 	UMovieGraphNode* OutputNode = TemporaryGraph->GetOutputNode();
@@ -557,7 +561,7 @@ UMovieGraphNode* FMovieGraphQuickRender::InjectNodeIntoBranch(const TSubclassOf<
 	return NewNode;
 }
 
-bool FMovieGraphQuickRender::AreModePrerequisitesMet(const EMovieGraphQuickRenderMode QuickRenderMode)
+bool UMovieGraphQuickRenderSubsystem::AreModePrerequisitesMet(const EMovieGraphQuickRenderMode QuickRenderMode) const
 {
 	if (QuickRenderMode == EMovieGraphQuickRenderMode::SelectedCameras)
 	{
@@ -577,7 +581,7 @@ bool FMovieGraphQuickRender::AreModePrerequisitesMet(const EMovieGraphQuickRende
 	return true;
 }
 
-void FMovieGraphQuickRender::CachePrePieData(UWorld* InEditorWorld)
+void UMovieGraphQuickRenderSubsystem::CachePrePieData(UWorld* InEditorWorld)
 {
 	// Prevent all existing level sequences in the world from auto-playing. We need the level sequence from Quick Render to be in control. 
 	for (auto It = TActorIterator<ALevelSequenceActor>(InEditorWorld); It; ++It)
@@ -596,7 +600,7 @@ void FMovieGraphQuickRender::CachePrePieData(UWorld* InEditorWorld)
 	CopyToWeakArray(CachedPrePieData.SelectedCameras, CameraActors);
 }
 
-void FMovieGraphQuickRender::RestorePreRenderState(const EMovieGraphQuickRenderMode QuickRenderMode)
+void UMovieGraphQuickRenderSubsystem::RestorePreRenderState(const EMovieGraphQuickRenderMode QuickRenderMode)
 {
 	// Restore the modified level sequences to the way they were before the render.
 	for (const TWeakObjectPtr<ALevelSequenceActor>& LevelSequenceActor : CachedPrePieData.ModifiedLevelSequenceActors)
@@ -629,7 +633,7 @@ void FMovieGraphQuickRender::RestorePreRenderState(const EMovieGraphQuickRenderM
 	}
 }
 
-void FMovieGraphQuickRender::SetUpRenderingLevelSequence(const EMovieGraphQuickRenderMode QuickRenderMode)
+void UMovieGraphQuickRenderSubsystem::SetUpRenderingLevelSequence(const EMovieGraphQuickRenderMode QuickRenderMode)
 {
 	auto GetCurrentLevelSequenceOrWarn = [](EMovieGraphQuickRenderMode ActiveQuickRenderMode) -> ULevelSequence*
 	{
@@ -654,11 +658,19 @@ void FMovieGraphQuickRender::SetUpRenderingLevelSequence(const EMovieGraphQuickR
 		ULevelSequence* QuickRenderSequence = NewObject<ULevelSequence>(GetTransientPackage(), TEXT("QuickRenderSequence"), RF_Transient);
 		QuickRenderSequence->Initialize();
 
-		RenderingLevelSequence.Reset(QuickRenderSequence);
+		RenderingLevelSequence = QuickRenderSequence;
 		
 		return;
 	}
 
+	// Scripting may override the level sequence to use
+	ULevelSequence* LevelSequenceOverride = QuickRenderModeSettings->LevelSequenceOverride.LoadSynchronous();
+	if (IsValid(LevelSequenceOverride))
+	{
+		RenderingLevelSequence = LevelSequenceOverride;
+		return;
+	}
+
 	// Otherwise, the active mode should use the level sequence that's currently active in Sequencer.
 	ULevelSequence* CurrentLevelSequence = GetCurrentLevelSequenceOrWarn(QuickRenderMode);
 	if (!CurrentLevelSequence)
@@ -666,18 +678,18 @@ void FMovieGraphQuickRender::SetUpRenderingLevelSequence(const EMovieGraphQuickR
 		return;
 	}
 	
-	RenderingLevelSequence.Reset(CurrentLevelSequence);
+	RenderingLevelSequence = CurrentLevelSequence;
 }
 
-void FMovieGraphQuickRender::SetUpUtilityLevelSequence()
+void UMovieGraphQuickRenderSubsystem::SetUpUtilityLevelSequence()
 {
 	ULevelSequence* UtilitySequence = NewObject<ULevelSequence>(GetTransientPackage(), TEXT("QuickRenderUtilitySequence"), RF_Transient);
 	UtilitySequence->Initialize();
 
-	UtilityLevelSequence.Reset(UtilitySequence);
+	UtilityLevelSequence = UtilitySequence;
 }
 
-UWorld* FMovieGraphQuickRender::GetWorldOfType(EWorldType::Type WorldType)
+UWorld* UMovieGraphQuickRenderSubsystem::GetWorldOfType(EWorldType::Type WorldType)
 {
 	UWorld* World = nullptr;
 	for (const FWorldContext& Context : GEngine->GetWorldContexts())
@@ -692,12 +704,12 @@ UWorld* FMovieGraphQuickRender::GetWorldOfType(EWorldType::Type WorldType)
 	return World;
 }
 
-TRange<FFrameNumber> FMovieGraphQuickRender::GetPlaybackRange(const UMovieGraphQuickRenderModeSettings& InQuickRenderSettings)
+TRange<FFrameNumber> UMovieGraphQuickRenderSubsystem::GetPlaybackRange()
 {
 	FFrameNumber StartFrame;
 	FFrameNumber EndFrame;
 
-	if (InQuickRenderSettings.FrameRangeType == EMovieGraphQuickRenderFrameRangeType::SelectionRange)
+	if (QuickRenderModeSettings->FrameRangeType == EMovieGraphQuickRenderFrameRangeType::SelectionRange)
 	{
 		const int32 SelectionStartFrame = ULevelSequenceEditorBlueprintLibrary::GetSelectionRangeStart();
 		const int32 SelectionEndFrame = ULevelSequenceEditorBlueprintLibrary::GetSelectionRangeEnd();
@@ -712,10 +724,10 @@ TRange<FFrameNumber> FMovieGraphQuickRender::GetPlaybackRange(const UMovieGraphQ
 			EndFrame = FFrameNumber(SelectionEndFrame);
 		}
 	}
-	else if (InQuickRenderSettings.FrameRangeType == EMovieGraphQuickRenderFrameRangeType::Custom)
+	else if (QuickRenderModeSettings->FrameRangeType == EMovieGraphQuickRenderFrameRangeType::Custom)
 	{
-		StartFrame = FFrameNumber(InQuickRenderSettings.CustomStartFrame);
-		EndFrame = FFrameNumber(InQuickRenderSettings.CustomEndFrame);
+		StartFrame = FFrameNumber(QuickRenderModeSettings->CustomStartFrame);
+		EndFrame = FFrameNumber(QuickRenderModeSettings->CustomEndFrame);
 	}
 	else
 	{
@@ -727,7 +739,7 @@ TRange<FFrameNumber> FMovieGraphQuickRender::GetPlaybackRange(const UMovieGraphQ
 	return TRange<FFrameNumber>::Inclusive(StartFrame, EndFrame);
 }
 
-void FMovieGraphQuickRender::PerformPreRenderSetup(const EMovieGraphQuickRenderMode QuickRenderMode, UWorld* InEditorWorld, const UMovieGraphQuickRenderModeSettings& InQuickRenderSettings)
+void UMovieGraphQuickRenderSubsystem::PerformPreRenderSetup(const EMovieGraphQuickRenderMode QuickRenderMode, UWorld* InEditorWorld)
 {
 	// Adds the specified camera to the Camera Cut track in the given level sequence.
 	// If there's already a camera cut section within the Camera Cut track, a new one will be added immediately after.
@@ -776,7 +788,7 @@ void FMovieGraphQuickRender::PerformPreRenderSetup(const EMovieGraphQuickRenderM
 	};
 	
 	// Spawn a new camera which mimics the viewport's camera 
-	auto DuplicateViewportCamera = []() -> ACameraActor*
+	auto DuplicateViewportCamera = [this]() -> ACameraActor*
 	{
 		// Get the PIE world, this is where we need to target the changes (not the editor world). Otherwise we'll have to deal with deleting actors
 		// after the render finishes, which isn't ideal (and this will also create entries in the undo stack, also not ideal).
@@ -879,7 +891,7 @@ void FMovieGraphQuickRender::PerformPreRenderSetup(const EMovieGraphQuickRenderM
 	// Adds a Level Visibility track to the Utility subsequence to ensure that the visibility state of levels in the editor is reflected
 	// in the render/PIE. Blueprint streamable levels, for example, may be shown in the editor, but will not typically show up in the
 	// render unless made visible by code or in Sequencer.
-	auto AddLevelVisibilityTracks = [InEditorWorld]() -> void
+	auto AddLevelVisibilityTracks = [this, InEditorWorld]() -> void
 	{
 		TArray<FName> VisibleLevelNames;
 		TArray<FName> HiddenLevelNames;
@@ -970,7 +982,7 @@ void FMovieGraphQuickRender::PerformPreRenderSetup(const EMovieGraphQuickRenderM
 	SequenceSubsection->SetRange(RenderingFrameRange);
 
 	// Add the Level Visibility tracks to the Utility sequence if matching the editor visibility was requested
-	const EMovieGraphQuickRenderViewportLookFlags LookFlags = static_cast<EMovieGraphQuickRenderViewportLookFlags>(InQuickRenderSettings.ViewportLookFlags);
+	const EMovieGraphQuickRenderViewportLookFlags LookFlags = static_cast<EMovieGraphQuickRenderViewportLookFlags>(QuickRenderModeSettings->ViewportLookFlags);
 	if (EnumHasAnyFlags(LookFlags, EMovieGraphQuickRenderViewportLookFlags::Visibility))
 	{
 		AddLevelVisibilityTracks();
diff --git a/Engine/Plugins/MovieScene/MovieRenderPipeline/Source/MovieRenderPipelineEditor/Private/MoviePipelineQuickRenderMenu.cpp b/Engine/Plugins/MovieScene/MovieRenderPipeline/Source/MovieRenderPipelineEditor/Private/MoviePipelineQuickRenderMenu.cpp
--- a/Engine/Plugins/MovieScene/MovieRenderPipeline/Source/MovieRenderPipelineEditor/Private/MoviePipelineQuickRenderMenu.cpp
+++ b/Engine/Plugins/MovieScene/MovieRenderPipeline/Source/MovieRenderPipelineEditor/Private/MoviePipelineQuickRenderMenu.cpp
@@ -99,7 +99,8 @@ void FMoviePipelineQuickRenderMenu::QuickRenderButtonPressed()
 		}
 		else
 		{
-			FMovieGraphQuickRender::BeginQuickRender(GetQuickRenderMode(), *QuickRenderModeSettings);
+			UMovieGraphQuickRenderSubsystem* QuickRenderSubsystem = GEditor->GetEditorSubsystem<UMovieGraphQuickRenderSubsystem>();
+			QuickRenderSubsystem->BeginQuickRender(GetQuickRenderMode(), QuickRenderModeSettings.Get());
 		}
 	}
 	else
@@ -287,7 +288,9 @@ TSharedRef<SWidget> FMoviePipelineQuickRenderMenu::GenerateQuickRenderSettingsWi
 			// Sync the Setting window's mode to the toolbar mode. This is only done when a render is performed, not when the window is closed.
 			SetQuickRenderMode(SettingsWindowQuickRenderMode);
 			
-			FMovieGraphQuickRender::BeginQuickRender(GetQuickRenderMode(), *QuickRenderModeSettings);
+			UMovieGraphQuickRenderSubsystem* QuickRenderSubsystem = GEditor->GetEditorSubsystem<UMovieGraphQuickRenderSubsystem>();
+			QuickRenderSubsystem->BeginQuickRender(GetQuickRenderMode(), QuickRenderModeSettings.Get());
+			
 			return FReply::Handled();
 		}),
 		NAME_None, false, HAlign_Right);
diff --git a/Engine/Plugins/MovieScene/MovieRenderPipeline/Source/MovieRenderPipelineEditor/Public/Graph/MovieGraphQuickRender.h b/Engine/Plugins/MovieScene/MovieRenderPipeline/Source/MovieRenderPipelineEditor/Public/Graph/MovieGraphQuickRender.h
--- a/Engine/Plugins/MovieScene/MovieRenderPipeline/Source/MovieRenderPipelineEditor/Public/Graph/MovieGraphQuickRender.h
+++ b/Engine/Plugins/MovieScene/MovieRenderPipeline/Source/MovieRenderPipelineEditor/Public/Graph/MovieGraphQuickRender.h
@@ -3,13 +3,15 @@
 #pragma once
 
 #include "CoreMinimal.h"
-#include "UObject/StrongObjectPtr.h"
+#include "EditorSubsystem.h"
 
 #include "Graph/MovieGraphQuickRenderSettings.h"
 #include "LevelSequence.h"
 #include "MoviePipelinePIEExecutor.h"
 #include "MoviePipelineQueue.h"
 
+#include "MovieGraphQuickRender.generated.h"
+
 class FEditorViewportClient;
 
 /**
@@ -17,11 +19,15 @@ class FEditorViewportClient;
  * creating a queue, adding a job(s) to it, specifying the level sequence to use, providing a custom configuration or graph, etc. Quick Renders are
  * designed to get you a render as fast as possible while providing minimal configuration input, for use in things like approving animation.
  */
-class MOVIERENDERPIPELINEEDITOR_API FMovieGraphQuickRender
+UCLASS()
+class MOVIERENDERPIPELINEEDITOR_API UMovieGraphQuickRenderSubsystem : public UEditorSubsystem
 {
+	GENERATED_BODY()
+
 public:
 	/* Begins a quick render using the provided mode and settings. */
-	static void BeginQuickRender(const EMovieGraphQuickRenderMode QuickRenderMode, const UMovieGraphQuickRenderModeSettings& InQuickRenderSettings);
+	UFUNCTION(BlueprintCallable, Category="Quick Render")
+	void BeginQuickRender(const EMovieGraphQuickRenderMode QuickRenderMode, const UMovieGraphQuickRenderModeSettings* InQuickRenderSettings);
 
 private:
 	/** Data that was cached prior to PIE starting. */
@@ -41,81 +47,94 @@ private:
 	 * Gets the graph that should be used for the quick render. This will be a duplicate of either the user-specified graph or the default quick render graph.
 	 * OutOriginalToDupeMap will provide a mapping of the pre-duplicate -> duplicate graphs (may contain multiple graphs if subgraphs are involved).
 	 */
-	static UMovieGraphConfig* GetQuickRenderGraph(UMovieGraphConfig* InUserSpecifiedGraph, TMap<TObjectPtr<UMovieGraphConfig>, TObjectPtr<UMovieGraphConfig>>& OutOriginalToDupeMap);
+	UMovieGraphConfig* GetQuickRenderGraph(UMovieGraphConfig* InUserSpecifiedGraph, TMap<TObjectPtr<UMovieGraphConfig>, TObjectPtr<UMovieGraphConfig>>& OutOriginalToDupeMap);
 
 	/** Applies all necessary updates to the duplicated graph to prepare it for a quick render (adjusting start/end frames, etc). */
-	static void ApplyQuickRenderUpdatesToDuplicateGraph(const UMovieGraphQuickRenderModeSettings& InQuickRenderSettings, UWorld* InEditorWorld, UMoviePipelineExecutorJob* InJob);
+	void ApplyQuickRenderUpdatesToDuplicateGraph(const UWorld* InEditorWorld);
 
 	/** Helper for ApplyQuickRenderUpdatesToDuplicateGraph() that applies actor editor visibility (by injecting a visibility modifier into the graph). */
-	static void ApplyQuickRenderUpdatesToDuplicateGraph_ApplyEditorVisibility(const UWorld* InEditorWorld);
+	void ApplyQuickRenderUpdatesToDuplicateGraph_ApplyEditorVisibility(const UWorld* InEditorWorld);
 
 	/** Helper for ApplyQuickRenderUpdatesToDuplicateGraph() that applies viewport OCIO (by injecting overrides to file output and renderer nodes into the graph). */
-	static void ApplyQuickRenderUpdatesToDuplicateGraph_ApplyOcio(const FOpenColorIODisplayConfiguration* InOcioConfiguration);
+	void ApplyQuickRenderUpdatesToDuplicateGraph_ApplyOcio(const FOpenColorIODisplayConfiguration* InOcioConfiguration);
 
 	/** Helper for ApplyQuickRenderUpdatesToDuplicateGraph() that applies viewport scalability if the graph didn't specify an explicit scalability setting. */
-	static void ApplyQuickRenderUpdatesToDuplicateGraph_Scalability();
+	void ApplyQuickRenderUpdatesToDuplicateGraph_Scalability();
 
 	/** Injects a new node in the furthest-downstream position within the specified branch in the duplicated graph. */
-	static UMovieGraphNode* InjectNodeIntoBranch(const TSubclassOf<UMovieGraphNode> NodeType, const FName& InBranchName);
+	UMovieGraphNode* InjectNodeIntoBranch(const TSubclassOf<UMovieGraphNode> NodeType, const FName& InBranchName) const;
 
 	/** Determines if the prerequisites for the given mode are met. */
-	static bool AreModePrerequisitesMet(const EMovieGraphQuickRenderMode QuickRenderMode);
+	bool AreModePrerequisitesMet(const EMovieGraphQuickRenderMode QuickRenderMode) const;
 
 	/** Caches any necessary data before PIE starts. Stored in the CachedPrePieData member variable. */
-	static void CachePrePieData(UWorld* InEditorWorld);
+	void CachePrePieData(UWorld* InEditorWorld);
 
 	/** Handles restoring data after PIE ends. */
-	static void RestorePreRenderState(const EMovieGraphQuickRenderMode QuickRenderMode);
+	void RestorePreRenderState(const EMovieGraphQuickRenderMode QuickRenderMode);
 
 	/**
 	 * Populates RenderingLevelSequence with the current level sequence, or sets up a new level sequence for rendering (depending on the mode that
 	 * Quick Render is using).
 	 */
-	static void SetUpRenderingLevelSequence(const EMovieGraphQuickRenderMode QuickRenderMode);
+	void SetUpRenderingLevelSequence(const EMovieGraphQuickRenderMode QuickRenderMode);
 
 	/**
 	 * Populates UtilityLevelSequence with a level sequence used for utility purposes (like overriding the camera in use, or changing level
 	 * visibility). It will be injected into the rendering level sequence as a subsequence.
 	 */
-	static void SetUpUtilityLevelSequence();
+	void SetUpUtilityLevelSequence();
 
 	/** Gets the first world found of the specified type. */
-	static UWorld* GetWorldOfType(EWorldType::Type WorldType);
+	UWorld* GetWorldOfType(EWorldType::Type WorldType);
 
 	/**
 	 * Gets the explicitly-set playback range of the level sequence being rendered (set either by a selection in Sequencer, or custom start/end
 	 * frames). If using the unaltered playback range of the level sequence, this will return an empty range.
 	 */
-	static TRange<FFrameNumber> GetPlaybackRange(const UMovieGraphQuickRenderModeSettings& InQuickRenderSettings);
+	TRange<FFrameNumber> GetPlaybackRange();
 
 	/** Does any setup needed in the PIE world prior to the render starting. */
-	static void PerformPreRenderSetup(const EMovieGraphQuickRenderMode QuickRenderMode, UWorld* InEditorWorld, const UMovieGraphQuickRenderModeSettings& InQuickRenderSettings);
+	void PerformPreRenderSetup(const EMovieGraphQuickRenderMode QuickRenderMode, UWorld* InEditorWorld);
 
 private:
 	/** The temporary queue that is used by quick render. */
-	inline static TStrongObjectPtr<UMoviePipelineQueue> TemporaryQueue = nullptr;
+	UPROPERTY(Transient)
+	TObjectPtr<UMoviePipelineQueue> TemporaryQueue = nullptr;
 
 	/** The temporary executor that drives the quick render. */
-	inline static TStrongObjectPtr<UMoviePipelinePIEExecutor> TemporaryExecutor = nullptr;
+	UPROPERTY(Transient)
+	TObjectPtr<UMoviePipelinePIEExecutor> TemporaryExecutor = nullptr;
 
 	/** The temporary graph that is used by quick render. Usually a duplicate of a graph created outside of quick render, but will be modified by quick render. */
-	inline static TStrongObjectPtr<UMovieGraphConfig> TemporaryGraph = nullptr;
+	UPROPERTY(Transient)
+	TObjectPtr<UMovieGraphConfig> TemporaryGraph = nullptr;
 
 	/** The temporary evaluated graph that is generated before a render starts; some setup processes need to inspect it. */
-	inline static TStrongObjectPtr<UMovieGraphEvaluatedConfig> TemporaryEvaluatedGraph = nullptr;
+	UPROPERTY(Transient)
+	TObjectPtr<UMovieGraphEvaluatedConfig> TemporaryEvaluatedGraph = nullptr;
 
 	/**
 	 * The level sequence that is used for rendering (ie, provided directly to MRG to render from). Sometimes this will be the current level sequence
 	 * in Sequencer, and other times it may be a temporary level sequence solely for use in Quick Render. It depends on the mode that is being used.
 	 */
-	inline static TStrongObjectPtr<ULevelSequence> RenderingLevelSequence = nullptr;
+	UPROPERTY(Transient)
+	TObjectPtr<ULevelSequence> RenderingLevelSequence = nullptr;
 	
 	/**
      * The level sequence that is being used by Quick Render for utility purposes. This is not directly rendered from, but may be included in another
      * level sequence to manipulate it (for example). For most modes, this will remain unused.
      */
-    inline static TStrongObjectPtr<ULevelSequence> UtilityLevelSequence = nullptr;
+    UPROPERTY(Transient)
+    TObjectPtr<ULevelSequence> UtilityLevelSequence = nullptr;
+
+	/**
+	 * The mode settings Quick Render was initialized with. These are cached here to grab a reference to them before the render begins so the
+	 * executor callbacks still have access to them later.
+	 */
+	UPROPERTY(Transient)
+	TObjectPtr<const UMovieGraphQuickRenderModeSettings> QuickRenderModeSettings = nullptr;
 
 	/** Any data that needed to be cached prior to starting PIE. */
-	inline static FCachedPrePieData CachedPrePieData;
+	FCachedPrePieData CachedPrePieData;
 };
