MoviePipeline: Updated the Quick Render menu to include two new options -- "Play Last Render" and "Open Output Directory".

diff --git a/Engine/Plugins/MovieScene/MovieRenderPipeline/Source/MovieRenderPipelineCore/Public/Graph/MovieGraphQuickRenderSettings.h b/Engine/Plugins/MovieScene/MovieRenderPipeline/Source/MovieRenderPipelineCore/Public/Graph/MovieGraphQuickRenderSettings.h
--- a/Engine/Plugins/MovieScene/MovieRenderPipeline/Source/MovieRenderPipelineCore/Public/Graph/MovieGraphQuickRenderSettings.h
+++ b/Engine/Plugins/MovieScene/MovieRenderPipeline/Source/MovieRenderPipelineCore/Public/Graph/MovieGraphQuickRenderSettings.h
@@ -61,8 +61,11 @@ enum class EMovieGraphQuickRenderPostRenderActionType : uint8
 	/** Do nothing after the render. */
 	DoNothing,
 
-	/** Open the movie that was generated from Quick Render. Opens via the operating system's default viewing application for the type of media that was created. */
-	PlayRenderOutput
+	/** Open the media that was generated from Quick Render. The application used for this is defined in Editor Preferences for Movie Render Graph. */
+	PlayRenderOutput,
+
+	/** Open the directory that contains the files generated by Quick Render. */
+	OpenOutputDirectory
 };
 
 /** The aspects of the viewport that should be applied to the quick render. */
diff --git a/Engine/Plugins/MovieScene/MovieRenderPipeline/Source/MovieRenderPipelineEditor/Private/Graph/MovieGraphQuickRender.cpp b/Engine/Plugins/MovieScene/MovieRenderPipeline/Source/MovieRenderPipelineEditor/Private/Graph/MovieGraphQuickRender.cpp
--- a/Engine/Plugins/MovieScene/MovieRenderPipeline/Source/MovieRenderPipelineEditor/Private/Graph/MovieGraphQuickRender.cpp
+++ b/Engine/Plugins/MovieScene/MovieRenderPipeline/Source/MovieRenderPipelineEditor/Private/Graph/MovieGraphQuickRender.cpp
@@ -6,11 +6,13 @@
 #include "Editor.h"
 #include "Engine/LevelStreaming.h"
 #include "EngineUtils.h"
+#include "HAL/FileManager.h"
 #include "LevelSequence.h"
 #include "LevelSequenceActor.h"
 #include "LevelSequenceEditorBlueprintLibrary.h"
 #include "LevelUtils.h"
 #include "Misc/MessageDialog.h"
+#include "MoviePipelineEditorBlueprintLibrary.h"
 #include "MoviePipelinePIEExecutor.h"
 #include "MoviePipelinePostRenderFileDisplayProcessor.h"
 #include "MoviePipelineQueueSubsystem.h"
@@ -90,16 +92,12 @@ void UMovieGraphQuickRenderSubsystem::BeginQuickRender(const EMovieGraphQuickRen
 	// Store a reference to the settings used, so the executor callbacks can rely on them still existing.
 	QuickRenderModeSettings = InQuickRenderSettings;
 
-	// Get/create the level sequence that should be rendered from (varies depending on the Quick Render mode).
-	SetUpRenderingLevelSequence(QuickRenderMode);
-	if (!IsValid(RenderingLevelSequence))
+	// Set up the rendering and "utility" level sequences. These are both required to do a render successfully.
+	if (!SetUpAllLevelSequences(QuickRenderMode))
 	{
 		return;
 	}
 
-	// Also set up the utility level sequence (which is used, for example, for overriding level visibility)
-	SetUpUtilityLevelSequence();
-
 	// Don't do additional work if the selected mode doesn't have what it needs to do a render. AreModePrerequisitesMet() will output the appropriate warnings.
 	if (!AreModePrerequisitesMet(QuickRenderMode))
 	{
@@ -114,46 +112,12 @@ void UMovieGraphQuickRenderSubsystem::BeginQuickRender(const EMovieGraphQuickRen
 		return;
 	}
 
-	// Allocate a new queue that will be used temporarily
-	TemporaryQueue = NewObject<UMoviePipelineQueue>(GetTransientPackage());
-
-	// Add a temp job to the queue, and give it the editor's level and the current level sequence being edited in Sequencer
-	UMoviePipelineExecutorJob* NewJob = TemporaryQueue->AllocateNewJob(UMoviePipelineExecutorJob::StaticClass());
-	NewJob->SetSequence(FSoftObjectPath(RenderingLevelSequence.Get()));
-	NewJob->Map = FSoftObjectPath(EditorWorld);
-	NewJob->JobName = FString(TEXT("QuickRender"));
-
-	// Update the job to use the quick render graph (either the default, or the one specified in the settings)
-	TMap<TObjectPtr<UMovieGraphConfig>, TObjectPtr<UMovieGraphConfig>> OriginalGraphToDupeMap;
-	const UMovieGraphConfig* QuickRenderGraph = GetQuickRenderGraph(QuickRenderModeSettings->GraphPreset.LoadSynchronous(), OriginalGraphToDupeMap);
-	if (!QuickRenderGraph)
+	// Generate a new queue, fully populated with a job to render from, the graph assigned to that job, and job-level variable overrides applied.
+	TObjectPtr<UMoviePipelineExecutorJob> NewJob = nullptr;
+	if (!GenerateAndPopulateQueue(TemporaryQueue, NewJob, QuickRenderModeSettings))
 	{
-		// GetQuickRenderGraph() will output an error to the log if a graph was not returned
 		return;
 	}
-	NewJob->SetGraphPreset(QuickRenderGraph);
-
-	// Apply any job-level variable overrides, if they were specified
-	for (const TObjectPtr<UMovieJobVariableAssignmentContainer>& IncomingVariableAssignment : QuickRenderModeSettings->GraphVariableAssignments)
-	{
-		// Map the setting's graph to the duplicate graph
-		const UMovieGraphConfig* AssignmentGraph = IncomingVariableAssignment->GetGraphConfig().LoadSynchronous();
-		const TObjectPtr<UMovieGraphConfig>* DuplicateGraph = OriginalGraphToDupeMap.Find(AssignmentGraph);
-		if (!DuplicateGraph)
-		{
-			UE_LOG(LogMovieRenderPipeline, Warning, TEXT("Quick Render: Unable to properly set variable overrides for graph [%s]."), *AssignmentGraph->GetName());
-			continue;
-		}
-		
-		UMovieJobVariableAssignmentContainer* NewVariableAssignments = NewJob->GetOrCreateJobVariableAssignmentsForGraph(*DuplicateGraph);
-
-		for (const UMovieGraphVariable* GraphVariable : (*DuplicateGraph)->GetVariables())
-		{
-			// Copying values by serialized string isn't ideal, but it's the easiest way to transfer values between the variable assignments
-			NewVariableAssignments->SetValueSerializedString(GraphVariable, IncomingVariableAssignment->GetValueSerializedString(GraphVariable));
-			NewVariableAssignments->SetVariableAssignmentEnableState(GraphVariable, true);
-		}
-	}
 	
 	// Evaluate the graph. Walking the graph to determine which nodes are present in each branch (especially taking into consideration
 	// subgraphs) is difficult here, so evaluation is the best way to deal with this.
@@ -194,33 +158,56 @@ void UMovieGraphQuickRenderSubsystem::BeginQuickRender(const EMovieGraphQuickRen
 	// View the job's rendered frame(s)/video(s) if requested
 	TemporaryExecutor->OnIndividualJobWorkFinished().AddLambda([this](FMoviePipelineOutputData OutputData)
 	{
-		if (QuickRenderModeSettings->PostRenderBehavior == EMovieGraphQuickRenderPostRenderActionType::PlayRenderOutput)
-		{
-			FMoviePipelinePostRenderFileDisplayProcessor PostRenderFileDisplayProcessor(GetDefault<UMovieRenderGraphEditorSettings>()->PostRenderSettings);
-			PostRenderFileDisplayProcessor.AddFiles(OutputData);
-			PostRenderFileDisplayProcessor.OpenFiles();
-		}
+		HandleJobFinished(QuickRenderModeSettings, OutputData);
 	});
 
 	// Clean up
 	TemporaryExecutor->OnExecutorFinished().AddLambda([this, QuickRenderMode](UMoviePipelineExecutorBase* InExecutor, bool bSuccess)
 	{
 		RestorePreRenderState(QuickRenderMode);
-
-		// The temp objects are part of the transient package, so manually release them when the render is finished
-		TemporaryQueue = nullptr;
-		TemporaryGraph = nullptr;
-		TemporaryExecutor = nullptr;
-		TemporaryEvaluatedGraph = nullptr;
-		RenderingLevelSequence = nullptr;
-		UtilityLevelSequence = nullptr;
-		QuickRenderModeSettings = nullptr;
+		PerformPostRenderCleanup();
 	});
 
 	// Do the render
 	Subsystem->RenderQueueInstanceWithExecutorInstance(TemporaryQueue.Get(), TemporaryExecutor.Get());
 }
 
+void UMovieGraphQuickRenderSubsystem::PlayLastRender()
+{
+	OpenPostRenderFileDisplayProcessor(PreviousRenderOutputData);
+}
+
+bool UMovieGraphQuickRenderSubsystem::CanPlayLastRender()
+{
+	return !PreviousRenderOutputData.GraphData.IsEmpty();
+}
+
+void UMovieGraphQuickRenderSubsystem::OpenOutputDirectory(const UMovieGraphQuickRenderModeSettings* InQuickRenderSettings)
+{
+	if (!InQuickRenderSettings)
+	{
+		return;
+	}
+
+	// In order to properly resolve the output directory, we have to set up a valid job.
+	TObjectPtr<UMoviePipelineExecutorJob> TempJob = nullptr;
+	TObjectPtr<UMoviePipelineQueue> TempQueue = nullptr;
+	if (!GenerateAndPopulateQueue(TempQueue, TempJob, InQuickRenderSettings))
+	{
+		return;
+	}
+
+	const FString ResolvedOutputDirectory = UMoviePipelineEditorBlueprintLibrary::ResolveOutputDirectoryFromJob(TempJob);
+
+	if (!ResolvedOutputDirectory.IsEmpty())
+	{
+		// The directory might not exist yet. Create it (if needed) so ExploreFolder() can open it.
+		IFileManager::Get().MakeDirectory(*ResolvedOutputDirectory, true);
+
+		FPlatformProcess::ExploreFolder(*ResolvedOutputDirectory);
+	}
+}
+
 UMovieGraphConfig* UMovieGraphQuickRenderSubsystem::GetQuickRenderGraph(UMovieGraphConfig* InUserSpecifiedGraph, TMap<TObjectPtr<UMovieGraphConfig>, TObjectPtr<UMovieGraphConfig>>& OutOriginalToDupeMap)
 {
 	UMovieGraphConfig* QuickRenderGraph = InUserSpecifiedGraph;
@@ -255,6 +242,77 @@ UMovieGraphConfig* UMovieGraphQuickRenderSubsystem::GetQuickRenderGraph(UMovieGr
 	return TemporaryGraph.Get();
 }
 
+bool UMovieGraphQuickRenderSubsystem::GenerateAndPopulateQueue(TObjectPtr<UMoviePipelineQueue>& OutQueue, TObjectPtr<UMoviePipelineExecutorJob>& OutJob, TObjectPtr<const UMovieGraphQuickRenderModeSettings> InQuickRenderSettings)
+{
+	if (!InQuickRenderSettings)
+	{
+		UE_LOG(LogMovieRenderPipeline, Warning, TEXT("Quick Render: Invalid mode settings provided."));
+		return false;
+	}
+	
+	UWorld* EditorWorld = GetWorldOfType(EWorldType::Editor);
+	if (!EditorWorld)
+	{
+		UE_LOG(LogMovieRenderPipeline, Warning, TEXT("Quick Render: Could not find an editor world."));
+		return false;
+	}
+	
+	// Allocate a new queue that will be used temporarily
+	OutQueue = NewObject<UMoviePipelineQueue>(GetTransientPackage());
+
+	// Add a temp job to the queue, and give it the editor's level and the current level sequence being edited in Sequencer
+	OutJob = OutQueue->AllocateNewJob(UMoviePipelineExecutorJob::StaticClass());
+	OutJob->SetSequence(FSoftObjectPath(RenderingLevelSequence.Get()));
+	OutJob->Map = FSoftObjectPath(EditorWorld);
+	OutJob->JobName = FString(TEXT("QuickRender"));
+
+	// Update the job to use the quick render graph (either the default, or the one specified in the settings)
+	TMap<TObjectPtr<UMovieGraphConfig>, TObjectPtr<UMovieGraphConfig>> OriginalGraphToDupeMap;
+	const UMovieGraphConfig* QuickRenderGraph = GetQuickRenderGraph(InQuickRenderSettings->GraphPreset.LoadSynchronous(), OriginalGraphToDupeMap);
+	if (!QuickRenderGraph)
+	{
+		// GetQuickRenderGraph() will output an error to the log if a graph was not returned
+		return false;
+	}
+	OutJob->SetGraphPreset(QuickRenderGraph);
+
+	// Apply any job-level variable overrides, if they were specified
+	for (const TObjectPtr<UMovieJobVariableAssignmentContainer>& IncomingVariableAssignment : InQuickRenderSettings->GraphVariableAssignments)
+	{
+		// Map the setting's graph to the duplicate graph
+		const UMovieGraphConfig* AssignmentGraph = IncomingVariableAssignment->GetGraphConfig().LoadSynchronous();
+		const TObjectPtr<UMovieGraphConfig>* DuplicateGraph = OriginalGraphToDupeMap.Find(AssignmentGraph);
+		if (!DuplicateGraph)
+		{
+			UE_LOG(LogMovieRenderPipeline, Warning, TEXT("Quick Render: Unable to properly set variable overrides for graph [%s]."), *AssignmentGraph->GetName());
+			continue;
+		}
+
+		UMovieJobVariableAssignmentContainer* NewVariableAssignments = OutJob->GetOrCreateJobVariableAssignmentsForGraph(*DuplicateGraph);
+
+		for (const UMovieGraphVariable* GraphVariable : (*DuplicateGraph)->GetVariables())
+		{
+			// Copying values by serialized string isn't ideal, but it's the easiest way to transfer values between the variable assignments
+			NewVariableAssignments->SetValueSerializedString(GraphVariable, IncomingVariableAssignment->GetValueSerializedString(GraphVariable));
+			NewVariableAssignments->SetVariableAssignmentEnableState(GraphVariable, true);
+		}
+	}
+
+	return true;
+}
+
+void UMovieGraphQuickRenderSubsystem::PerformPostRenderCleanup()
+{
+	// The temp objects are part of the transient package, so manually release them when the render is finished
+	TemporaryQueue = nullptr;
+	TemporaryGraph = nullptr;
+	TemporaryExecutor = nullptr;
+	TemporaryEvaluatedGraph = nullptr;
+	RenderingLevelSequence = nullptr;
+	UtilityLevelSequence = nullptr;
+	QuickRenderModeSettings = nullptr;
+}
+
 void UMovieGraphQuickRenderSubsystem::ApplyQuickRenderUpdatesToDuplicateGraph(const UWorld* InEditorWorld)
 {
 	check(IsValid(TemporaryGraph));
@@ -633,7 +691,7 @@ void UMovieGraphQuickRenderSubsystem::RestorePreRenderState(const EMovieGraphQui
 	}
 }
 
-void UMovieGraphQuickRenderSubsystem::SetUpRenderingLevelSequence(const EMovieGraphQuickRenderMode QuickRenderMode)
+ULevelSequence* UMovieGraphQuickRenderSubsystem::SetUpRenderingLevelSequence(const EMovieGraphQuickRenderMode QuickRenderMode) const
 {
 	auto GetCurrentLevelSequenceOrWarn = [](EMovieGraphQuickRenderMode ActiveQuickRenderMode) -> ULevelSequence*
 	{
@@ -658,35 +716,43 @@ void UMovieGraphQuickRenderSubsystem::SetUpRenderingLevelSequence(const EMovieGr
 		ULevelSequence* QuickRenderSequence = NewObject<ULevelSequence>(GetTransientPackage(), TEXT("QuickRenderSequence"), RF_Transient);
 		QuickRenderSequence->Initialize();
 
-		RenderingLevelSequence = QuickRenderSequence;
-		
-		return;
+		return QuickRenderSequence;
 	}
 
 	// Scripting may override the level sequence to use
 	ULevelSequence* LevelSequenceOverride = QuickRenderModeSettings->LevelSequenceOverride.LoadSynchronous();
 	if (IsValid(LevelSequenceOverride))
 	{
-		RenderingLevelSequence = LevelSequenceOverride;
-		return;
+		return LevelSequenceOverride;
 	}
 
 	// Otherwise, the active mode should use the level sequence that's currently active in Sequencer.
 	ULevelSequence* CurrentLevelSequence = GetCurrentLevelSequenceOrWarn(QuickRenderMode);
 	if (!CurrentLevelSequence)
 	{
-		return;
+		return nullptr;
 	}
 	
-	RenderingLevelSequence = CurrentLevelSequence;
+	return CurrentLevelSequence;
 }
 
-void UMovieGraphQuickRenderSubsystem::SetUpUtilityLevelSequence()
+ULevelSequence* UMovieGraphQuickRenderSubsystem::SetUpUtilityLevelSequence() const
 {
 	ULevelSequence* UtilitySequence = NewObject<ULevelSequence>(GetTransientPackage(), TEXT("QuickRenderUtilitySequence"), RF_Transient);
 	UtilitySequence->Initialize();
 
-	UtilityLevelSequence = UtilitySequence;
+	return UtilitySequence;
+}
+
+bool UMovieGraphQuickRenderSubsystem::SetUpAllLevelSequences(const EMovieGraphQuickRenderMode QuickRenderMode)
+{
+	// Get/create the level sequence that should be rendered from (varies depending on the Quick Render mode).
+	RenderingLevelSequence = SetUpRenderingLevelSequence(QuickRenderMode);
+
+	// Also set up the utility level sequence (which is used, for example, for overriding level visibility)
+	UtilityLevelSequence = SetUpUtilityLevelSequence();
+
+	return IsValid(RenderingLevelSequence) && IsValid(UtilityLevelSequence);
 }
 
 UWorld* UMovieGraphQuickRenderSubsystem::GetWorldOfType(EWorldType::Type WorldType)
@@ -989,4 +1055,25 @@ void UMovieGraphQuickRenderSubsystem::PerformPreRenderSetup(const EMovieGraphQui
 	}
 }
 
+void UMovieGraphQuickRenderSubsystem::OpenPostRenderFileDisplayProcessor(const FMoviePipelineOutputData& InOutputData) const
+{
+	FMoviePipelinePostRenderFileDisplayProcessor PostRenderFileDisplayProcessor(GetDefault<UMovieRenderGraphEditorSettings>()->PostRenderSettings);
+	PostRenderFileDisplayProcessor.AddFiles(InOutputData);
+	PostRenderFileDisplayProcessor.OpenFiles();
+}
+
+void UMovieGraphQuickRenderSubsystem::HandleJobFinished(const UMovieGraphQuickRenderModeSettings* InQuickRenderSettings, const FMoviePipelineOutputData& InGeneratedOutputData)
+{
+	if (InQuickRenderSettings->PostRenderBehavior == EMovieGraphQuickRenderPostRenderActionType::PlayRenderOutput)
+	{
+		OpenPostRenderFileDisplayProcessor(InGeneratedOutputData);
+	}
+	else if (InQuickRenderSettings->PostRenderBehavior == EMovieGraphQuickRenderPostRenderActionType::OpenOutputDirectory)
+	{
+		OpenOutputDirectory(InQuickRenderSettings);
+	}
+
+	PreviousRenderOutputData = InGeneratedOutputData;
+}
+
 #undef LOCTEXT_NAMESPACE
diff --git a/Engine/Plugins/MovieScene/MovieRenderPipeline/Source/MovieRenderPipelineEditor/Private/MoviePipelineEditorBlueprintLibrary.cpp b/Engine/Plugins/MovieScene/MovieRenderPipeline/Source/MovieRenderPipelineEditor/Private/MoviePipelineEditorBlueprintLibrary.cpp
--- a/Engine/Plugins/MovieScene/MovieRenderPipeline/Source/MovieRenderPipelineEditor/Private/MoviePipelineEditorBlueprintLibrary.cpp
+++ b/Engine/Plugins/MovieScene/MovieRenderPipeline/Source/MovieRenderPipelineEditor/Private/MoviePipelineEditorBlueprintLibrary.cpp
@@ -22,6 +22,7 @@
 #include "MoviePipelineBlueprintLibrary.h"
 #include "MoviePipelineOutputSetting.h"
 #include "Graph/MovieGraphBlueprintLibrary.h"
+#include "Graph/Nodes/MovieGraphGlobalOutputSettingNode.h"
 
 #include UE_INLINE_GENERATED_CPP_BY_NAME(MoviePipelineEditorBlueprintLibrary)
 
@@ -246,29 +247,45 @@ bool UMoviePipelineEditorBlueprintLibrary::GetDisplayOutputPathFromJob(UMoviePip
 
 FString UMoviePipelineEditorBlueprintLibrary::ResolveOutputDirectoryFromJob(UMoviePipelineExecutorJob* InJob)
 {
-	// Set up as many parameters as we can to try and resolve most of the string.
+	UMovieGraphEvaluatedConfig* EvaluatedGraph = nullptr;
 	FString FormatString;
-	GetDisplayOutputPathFromJob(InJob, FormatString);
+	
+	// Set up as many parameters as we can to try and resolve most of the string.
+	// If using the graph, don't rely on GetDisplayOutputPathFromJob() because it may not take into account variables and other more complicated ways
+	// the output path can be set. It's good for display purposes because it runs fast, but since ResolveOutputDirectoryFromJob() is not called
+	// continuously from the UI, we can take a slower approach here and ensure that the fully correct path is fetched. This needs the evaluated graph.
+	if (InJob->IsUsingGraphConfiguration())
+	{
+		FString OutTraversalError;
+		FMovieGraphTraversalContext TraversalContext;
+		TraversalContext.Job = InJob;
+
+		EvaluatedGraph = InJob->GetGraphPreset()->CreateFlattenedGraph(TraversalContext, OutTraversalError);
+
+		constexpr bool bIncludeCDOs = true;
+		constexpr bool bExactMatch = true;
+		UMovieGraphGlobalOutputSettingNode* OutputSettings = EvaluatedGraph->GetSettingForBranch<UMovieGraphGlobalOutputSettingNode>(UMovieGraphNode::GlobalsPinName, bIncludeCDOs, bExactMatch);
+		FormatString = OutputSettings->OutputDirectory.Path;
+	}
+	else
+	{
+		GetDisplayOutputPathFromJob(InJob, FormatString);
+	}
 
 	FString OutResolvedPath;
 
 	FPaths::NormalizeFilename(FormatString);
 
-	if (InJob->IsUsingGraphConfiguration() && InJob->GetGraphPreset())
+	if (InJob->IsUsingGraphConfiguration())
 	{
 		const bool bGetNextVersion = false;
 		FMovieGraphFilenameResolveParams Params;
 		Params.Job = InJob;
 		Params.InitializationTime = FDateTime::UtcNow();
 		Params.InitializationTimeOffset = FDateTime::Now() - FDateTime::UtcNow();
-
 		Params.Version = UMovieGraphBlueprintLibrary::ResolveVersionNumber(Params, bGetNextVersion);
 		Params.RenderDataIdentifier.RootBranchName = UMovieGraphNode::GlobalsPinName;
-		
-		FString OutTraversalError;
-		FMovieGraphTraversalContext TraversalContext;
-		TraversalContext.Job = InJob;
-		Params.EvaluatedConfig = InJob->GetGraphPreset()->CreateFlattenedGraph(TraversalContext, OutTraversalError);
+		Params.EvaluatedConfig = EvaluatedGraph;
 
 		FMovieGraphResolveArgs Dummy;
 		OutResolvedPath = UMovieGraphBlueprintLibrary::ResolveFilenameFormatArguments(FormatString, Params, Dummy);
diff --git a/Engine/Plugins/MovieScene/MovieRenderPipeline/Source/MovieRenderPipelineEditor/Private/MoviePipelineQuickRenderMenu.cpp b/Engine/Plugins/MovieScene/MovieRenderPipeline/Source/MovieRenderPipelineEditor/Private/MoviePipelineQuickRenderMenu.cpp
--- a/Engine/Plugins/MovieScene/MovieRenderPipeline/Source/MovieRenderPipelineEditor/Private/MoviePipelineQuickRenderMenu.cpp
+++ b/Engine/Plugins/MovieScene/MovieRenderPipeline/Source/MovieRenderPipelineEditor/Private/MoviePipelineQuickRenderMenu.cpp
@@ -135,7 +135,7 @@ TSharedRef<SWidget> FMoviePipelineQuickRenderMenu::GenerateQuickRenderOptionsMen
 		GenerateModesMenuSection(Menu);
 		GenerateQuickRenderMenuSection(Menu);
 		GenerateQuickRenderConfigurationMenuSection(Menu);
-		GenerateCreateMenuSection(Menu);
+		GenerateOutputMenuSection(Menu);
 		GenerateSettingsMenuSection(Menu);
 	}
 	
@@ -456,19 +456,50 @@ void FMoviePipelineQuickRenderMenu::GenerateQuickRenderConfigurationMenuSection(
 	));
 }
 
-void FMoviePipelineQuickRenderMenu::GenerateCreateMenuSection(UToolMenu* InMenu)
+void FMoviePipelineQuickRenderMenu::GenerateOutputMenuSection(UToolMenu* InMenu)
 {
-	FToolMenuSection& Section = InMenu->AddSection("MoviePipelineQuickRenderCreate", LOCTEXT("QuickRenderCreate", "Create"));
-		
-	FToolUIAction CreateSequenceAction;
-	CreateSequenceAction.ExecuteAction = FToolMenuExecuteAction::CreateLambda([](const FToolMenuContext&) { return; });
-	CreateSequenceAction.CanExecuteAction = FToolMenuCanExecuteAction::CreateLambda([](const FToolMenuContext&) { return true; });
+	// Only show the section if Quick Render mode is active
+	FToolMenuSection& Section = InMenu->AddSection("MoviePipelineQuickRenderOutput", LOCTEXT("QuickRenderOutput", "Output"));
+	Section.Visibility = TAttribute<EVisibility>::CreateLambda([]() -> EVisibility
+	{
+		return (GetQuickRenderButtonMode() == EMovieGraphQuickRenderButtonMode::QuickRender) ? EVisibility::Visible : EVisibility::Collapsed;
+	});
+
+	// Menu entry for "Play Last Render"
+	// ---------------------------------
+	FToolUIAction PlayLastRenderAction;
+	PlayLastRenderAction.ExecuteAction = FToolMenuExecuteAction::CreateLambda([](const FToolMenuContext&)
+	{
+		UMovieGraphQuickRenderSubsystem* QuickRenderSubsystem = GEditor->GetEditorSubsystem<UMovieGraphQuickRenderSubsystem>();
+		QuickRenderSubsystem->PlayLastRender();
+	});
+	PlayLastRenderAction.CanExecuteAction = FToolMenuCanExecuteAction::CreateLambda([](const FToolMenuContext&)
+	{
+		UMovieGraphQuickRenderSubsystem* QuickRenderSubsystem = GEditor->GetEditorSubsystem<UMovieGraphQuickRenderSubsystem>();
+		return QuickRenderSubsystem->CanPlayLastRender();
+	});
+	Section.AddMenuEntry(
+		"QuickRender_PlayLastRender",
+		LOCTEXT("QuickRender_PlayLastRenderLabel", "Play Last Render"),
+		LOCTEXT("QuickRender_PlayLastRenderToolTip", "Play the media from the last time that Quick Render ran."),
+		FSlateIcon(FAppStyle::GetAppStyleSetName(), "GenericPlay"),
+		PlayLastRenderAction);
+
+	// Menu entry for "Open Output Directory"
+	// --------------------------------------
+	FToolUIAction OpenOutputDirectoryAction;
+	OpenOutputDirectoryAction.ExecuteAction = FToolMenuExecuteAction::CreateLambda([](const FToolMenuContext&)
+	{
+		UMovieGraphQuickRenderSubsystem* QuickRenderSubsystem = GEditor->GetEditorSubsystem<UMovieGraphQuickRenderSubsystem>();
+		QuickRenderSubsystem->OpenOutputDirectory(QuickRenderModeSettings.Get());
+	});
+	OpenOutputDirectoryAction.CanExecuteAction = FToolMenuCanExecuteAction::CreateLambda([](const FToolMenuContext&) { return true; });
 	Section.AddMenuEntry(
-		"QuickRender_CreateSequenceFromCurrentViewport",
-		LOCTEXT("QuickRender_CreateSequenceFromCurrentViewportLabel", "Create Sequence from Current Viewport"),
-		LOCTEXT("QuickRender_CreateSequenceFromCurrentViewportToolTip", "Create a sequence from the current viewport."),
-		FSlateIcon(),
-		CreateSequenceAction);
+		"QuickRender_OpenOutputDirectory",
+		LOCTEXT("QuickRender_OpenOutputDirectoryLabel", "Open Output Directory"),
+		LOCTEXT("QuickRender_OpenOutputDirectoryToolTip", "Open the output directory that Quick Render is currently configured to save media into."),
+		FSlateIcon(FAppStyle::GetAppStyleSetName(), "Icons.FolderOpen"),
+		OpenOutputDirectoryAction);
 }
 
 void FMoviePipelineQuickRenderMenu::GenerateSettingsMenuSection(UToolMenu* InMenu)
diff --git a/Engine/Plugins/MovieScene/MovieRenderPipeline/Source/MovieRenderPipelineEditor/Private/MoviePipelineQuickRenderMenu.h b/Engine/Plugins/MovieScene/MovieRenderPipeline/Source/MovieRenderPipelineEditor/Private/MoviePipelineQuickRenderMenu.h
--- a/Engine/Plugins/MovieScene/MovieRenderPipeline/Source/MovieRenderPipelineEditor/Private/MoviePipelineQuickRenderMenu.h
+++ b/Engine/Plugins/MovieScene/MovieRenderPipeline/Source/MovieRenderPipelineEditor/Private/MoviePipelineQuickRenderMenu.h
@@ -35,7 +35,7 @@ private:
 	static void GenerateModesMenuSection(UToolMenu* InMenu);
 	static void GenerateQuickRenderMenuSection(UToolMenu* InMenu);
 	static void GenerateQuickRenderConfigurationMenuSection(UToolMenu* InMenu);
-	static void GenerateCreateMenuSection(UToolMenu* InMenu);
+	static void GenerateOutputMenuSection(UToolMenu* InMenu);
 	static void GenerateSettingsMenuSection(UToolMenu* InMenu);
 
 	static FSlateIcon GetIconForQuickRenderMode(EMovieGraphQuickRenderMode QuickRenderMode);
diff --git a/Engine/Plugins/MovieScene/MovieRenderPipeline/Source/MovieRenderPipelineEditor/Public/Graph/MovieGraphQuickRender.h b/Engine/Plugins/MovieScene/MovieRenderPipeline/Source/MovieRenderPipelineEditor/Public/Graph/MovieGraphQuickRender.h
--- a/Engine/Plugins/MovieScene/MovieRenderPipeline/Source/MovieRenderPipelineEditor/Public/Graph/MovieGraphQuickRender.h
+++ b/Engine/Plugins/MovieScene/MovieRenderPipeline/Source/MovieRenderPipelineEditor/Public/Graph/MovieGraphQuickRender.h
@@ -29,6 +29,24 @@ public:
 	UFUNCTION(BlueprintCallable, Category="Quick Render")
 	void BeginQuickRender(const EMovieGraphQuickRenderMode QuickRenderMode, const UMovieGraphQuickRenderModeSettings* InQuickRenderSettings);
 
+	/**
+	 * Plays the last render that Quick Render generated, using the settings specified in Editor Preferences. If no render has been
+	 * generated in this editor session yet, this does nothing.
+	 */
+	UFUNCTION(BlueprintCallable, Category="Quick Render")
+	void PlayLastRender();
+
+	/** Determines if the last render can be played (ie, whether a render has been performed in the current editor session or not). */
+	UFUNCTION(BlueprintCallable, Category="Quick Render")
+	bool CanPlayLastRender();
+
+	/** 
+	 * Opens the output directory that Quick Render will save media into. This is sourced from the graph that is used for rendering, which is why
+	 * settings need to be provided.
+	 */
+	UFUNCTION(BlueprintCallable, Category="Quick Render")
+	void OpenOutputDirectory(const UMovieGraphQuickRenderModeSettings* InQuickRenderSettings);
+
 private:
 	/** Data that was cached prior to PIE starting. */
 	struct FCachedPrePieData
@@ -49,6 +67,12 @@ private:
 	 */
 	UMovieGraphConfig* GetQuickRenderGraph(UMovieGraphConfig* InUserSpecifiedGraph, TMap<TObjectPtr<UMovieGraphConfig>, TObjectPtr<UMovieGraphConfig>>& OutOriginalToDupeMap);
 
+	/** Generates a new queue for use with Quick Render. Adds a job to the queue, specifies the graph, and all other setup required to get the queue operational. */
+	bool GenerateAndPopulateQueue(TObjectPtr<UMoviePipelineQueue>& OutQueue, TObjectPtr<UMoviePipelineExecutorJob>& OutJob, TObjectPtr<const UMovieGraphQuickRenderModeSettings> InQuickRenderSettings);
+
+	/** Performs the required post-render clean-up, mostly clearing out temporary data members. */
+	void PerformPostRenderCleanup();
+
 	/** Applies all necessary updates to the duplicated graph to prepare it for a quick render (adjusting start/end frames, etc). */
 	void ApplyQuickRenderUpdatesToDuplicateGraph(const UWorld* InEditorWorld);
 
@@ -74,16 +98,22 @@ private:
 	void RestorePreRenderState(const EMovieGraphQuickRenderMode QuickRenderMode);
 
 	/**
-	 * Populates RenderingLevelSequence with the current level sequence, or sets up a new level sequence for rendering (depending on the mode that
+	 * Returns the current level sequence active in Sequencer, or sets up a new level sequence for rendering (depending on the mode that
 	 * Quick Render is using).
 	 */
-	void SetUpRenderingLevelSequence(const EMovieGraphQuickRenderMode QuickRenderMode);
+	ULevelSequence* SetUpRenderingLevelSequence(const EMovieGraphQuickRenderMode QuickRenderMode) const;
 
 	/**
-	 * Populates UtilityLevelSequence with a level sequence used for utility purposes (like overriding the camera in use, or changing level
+	 * Returns a new level sequence used for utility purposes (like overriding the camera in use, or changing level
 	 * visibility). It will be injected into the rendering level sequence as a subsequence.
 	 */
-	void SetUpUtilityLevelSequence();
+	ULevelSequence* SetUpUtilityLevelSequence() const;
+
+	/**
+	 * Populates both the rendering and utility level sequences, and assigns them to the relevant members. Returns true if everything could be set
+	 * up, otherwise false.
+	 */
+	bool SetUpAllLevelSequences(const EMovieGraphQuickRenderMode QuickRenderMode);
 
 	/** Gets the first world found of the specified type. */
 	UWorld* GetWorldOfType(EWorldType::Type WorldType);
@@ -97,6 +127,12 @@ private:
 	/** Does any setup needed in the PIE world prior to the render starting. */
 	void PerformPreRenderSetup(const EMovieGraphQuickRenderMode QuickRenderMode, UWorld* InEditorWorld);
 
+	/** Opens the rendered files that are in the given output data. */
+	void OpenPostRenderFileDisplayProcessor(const FMoviePipelineOutputData& InOutputData) const;
+
+	/** Takes care of post-render tasks, like opening the rendered files. */
+	void HandleJobFinished(const UMovieGraphQuickRenderModeSettings* InQuickRenderSettings, const FMoviePipelineOutputData& InGeneratedOutputData);
+
 private:
 	/** The temporary queue that is used by quick render. */
 	UPROPERTY(Transient)
@@ -137,4 +173,7 @@ private:
 
 	/** Any data that needed to be cached prior to starting PIE. */
 	FCachedPrePieData CachedPrePieData;
+
+	/** Cache of the output data from the last render. */
+	FMoviePipelineOutputData PreviousRenderOutputData;
 };
