MoviePipeline: Updated Quick Render to ensure spawnables are visible in the "Current Viewport" and "Selected Camera(s)" modes.

diff --git a/Engine/Plugins/MovieScene/MovieRenderPipeline/Source/MovieRenderPipelineEditor/Private/Graph/MovieGraphQuickRender.cpp b/Engine/Plugins/MovieScene/MovieRenderPipeline/Source/MovieRenderPipelineEditor/Private/Graph/MovieGraphQuickRender.cpp
--- a/Engine/Plugins/MovieScene/MovieRenderPipeline/Source/MovieRenderPipelineEditor/Private/Graph/MovieGraphQuickRender.cpp
+++ b/Engine/Plugins/MovieScene/MovieRenderPipeline/Source/MovieRenderPipelineEditor/Private/Graph/MovieGraphQuickRender.cpp
@@ -20,6 +20,7 @@
 #include "MovieRenderPipelineCoreModule.h"
 #include "MovieRenderPipelineSettings.h"
 #include "MovieScene.h"
+#include "MovieSceneSpawnableAnnotation.h"
 #include "ObjectEditorUtils.h"
 #include "Scalability.h"
 #include "Sections/MovieSceneCameraCutSection.h"
@@ -352,7 +353,7 @@ void UMovieGraphQuickRenderSubsystem::ApplyQuickRenderUpdatesToDuplicateGraph(co
 		check(ApplyViewportLookNode);
 
 #define APPLY_VIEWPORT_LOOK_FLAG(EnumName, FlagName) \
-		if (EnumHasAnyFlags(LookFlags, EnumName)) \
+		if (IsViewportLookFlagActive(EnumName)) \
 		{ \
 			ApplyViewportLookNode->bOverride_b##FlagName = true; \
 			ApplyViewportLookNode->b##FlagName = true; \
@@ -367,13 +368,13 @@ void UMovieGraphQuickRenderSubsystem::ApplyQuickRenderUpdatesToDuplicateGraph(co
 	}
 
 	// Apply viewport/editor actor visibility
-	if (EnumHasAnyFlags(LookFlags, EMovieGraphQuickRenderViewportLookFlags::Visibility))
+	if (IsViewportLookFlagActive(EMovieGraphQuickRenderViewportLookFlags::Visibility))
 	{
 		ApplyQuickRenderUpdatesToDuplicateGraph_ApplyEditorVisibility(InEditorWorld);
 	}
 
 	// Apply OCIO if it is activated on the viewport and enabled in the viewport look flags
-	if (QuickRenderModeSettings->bOverride_ViewportLookFlags && EnumHasAnyFlags(LookFlags, EMovieGraphQuickRenderViewportLookFlags::Ocio))
+	if (IsViewportLookFlagActive(EMovieGraphQuickRenderViewportLookFlags::Ocio))
 	{
 #if WITH_OCIO
 		if (const FEditorViewportClient* ViewportClient = UMovieGraphApplyViewportLookNode::GetViewportClient())
@@ -736,11 +737,51 @@ ULevelSequence* UMovieGraphQuickRenderSubsystem::SetUpRenderingLevelSequence(con
 	return CurrentLevelSequence;
 }
 
-ULevelSequence* UMovieGraphQuickRenderSubsystem::SetUpUtilityLevelSequence() const
+ULevelSequence* UMovieGraphQuickRenderSubsystem::SetUpUtilityLevelSequence(const EMovieGraphQuickRenderMode QuickRenderMode) const
 {
 	ULevelSequence* UtilitySequence = NewObject<ULevelSequence>(GetTransientPackage(), TEXT("QuickRenderUtilitySequence"), RF_Transient);
 	UtilitySequence->Initialize();
 
+	// If the currently active level sequence in Sequencer has spawnables, mirror those spawnable bindings into the utility level sequence. For modes
+	// that don't use the active level sequence, we need to do this in order to match visibility.
+	auto AddSpawnableTracks = [this, UtilitySequence]() -> void
+	{
+		const UWorld* EditorWorld = GetWorldOfType(EWorldType::Editor);
+		const UMovieScene* UtilityMovieScene = UtilitySequence->GetMovieScene();
+
+		for (FActorIterator It(EditorWorld); It; ++It)
+		{
+			AActor* Actor = *It;
+
+			// There's no need to make the actor spawnable if it's currently not visible in the editor
+			if (!Actor || Actor->IsHiddenEd())
+			{
+				continue;
+			}
+
+			TOptional<FMovieSceneSpawnableAnnotation> Spawnable = FMovieSceneSpawnableAnnotation::Find(Actor);
+			if (Spawnable.IsSet())
+			{
+				FMovieGraphTransactionBlocker BlockTransactions;
+				
+				UE::Sequencer::FCreateBindingParams CreateBindingParams;
+				CreateBindingParams.bAllowCustomBinding = true;
+				CreateBindingParams.bSpawnable = true;
+				FSequencerUtilities::CreateOrReplaceBinding(nullptr, UtilityMovieScene->GetTypedOuter<UMovieSceneSequence>(), Actor, CreateBindingParams);
+			}
+		}
+	};
+
+	if (IsViewportLookFlagActive(EMovieGraphQuickRenderViewportLookFlags::Visibility))
+	{
+		// Only modes that don't use the active level sequence need the spawnable bindings mirrored. Ideally this could be done in
+		// PerformPreRenderSetup(), but at that point it's too late to get spawnable state because PIE has started.
+		if ((QuickRenderMode == EMovieGraphQuickRenderMode::CurrentViewport) || (QuickRenderMode == EMovieGraphQuickRenderMode::SelectedCameras))
+		{
+			AddSpawnableTracks();
+		}
+	}
+
 	return UtilitySequence;
 }
 
@@ -750,12 +791,12 @@ bool UMovieGraphQuickRenderSubsystem::SetUpAllLevelSequences(const EMovieGraphQu
 	RenderingLevelSequence = SetUpRenderingLevelSequence(QuickRenderMode);
 
 	// Also set up the utility level sequence (which is used, for example, for overriding level visibility)
-	UtilityLevelSequence = SetUpUtilityLevelSequence();
+	UtilityLevelSequence = SetUpUtilityLevelSequence(QuickRenderMode);
 
 	return IsValid(RenderingLevelSequence) && IsValid(UtilityLevelSequence);
 }
 
-UWorld* UMovieGraphQuickRenderSubsystem::GetWorldOfType(EWorldType::Type WorldType)
+UWorld* UMovieGraphQuickRenderSubsystem::GetWorldOfType(const EWorldType::Type WorldType) const
 {
 	UWorld* World = nullptr;
 	for (const FWorldContext& Context : GEngine->GetWorldContexts())
@@ -1049,7 +1090,7 @@ void UMovieGraphQuickRenderSubsystem::PerformPreRenderSetup(const EMovieGraphQui
 
 	// Add the Level Visibility tracks to the Utility sequence if matching the editor visibility was requested
 	const EMovieGraphQuickRenderViewportLookFlags LookFlags = static_cast<EMovieGraphQuickRenderViewportLookFlags>(QuickRenderModeSettings->ViewportLookFlags);
-	if (EnumHasAnyFlags(LookFlags, EMovieGraphQuickRenderViewportLookFlags::Visibility))
+	if (IsViewportLookFlagActive(EMovieGraphQuickRenderViewportLookFlags::Visibility))
 	{
 		AddLevelVisibilityTracks();
 	}
@@ -1076,4 +1117,16 @@ void UMovieGraphQuickRenderSubsystem::HandleJobFinished(const UMovieGraphQuickRe
 	PreviousRenderOutputData = InGeneratedOutputData;
 }
 
+bool UMovieGraphQuickRenderSubsystem::IsViewportLookFlagActive(const EMovieGraphQuickRenderViewportLookFlags ViewportLookFlag) const
+{
+	if (!IsValid(QuickRenderModeSettings))
+	{
+		return false;
+	}
+	
+	const EMovieGraphQuickRenderViewportLookFlags LookFlags = static_cast<EMovieGraphQuickRenderViewportLookFlags>(QuickRenderModeSettings->ViewportLookFlags);
+	
+	return QuickRenderModeSettings->bOverride_ViewportLookFlags && EnumHasAnyFlags(LookFlags, ViewportLookFlag);
+}
+
 #undef LOCTEXT_NAMESPACE
diff --git a/Engine/Plugins/MovieScene/MovieRenderPipeline/Source/MovieRenderPipelineEditor/Public/Graph/MovieGraphQuickRender.h b/Engine/Plugins/MovieScene/MovieRenderPipeline/Source/MovieRenderPipelineEditor/Public/Graph/MovieGraphQuickRender.h
--- a/Engine/Plugins/MovieScene/MovieRenderPipeline/Source/MovieRenderPipelineEditor/Public/Graph/MovieGraphQuickRender.h
+++ b/Engine/Plugins/MovieScene/MovieRenderPipeline/Source/MovieRenderPipelineEditor/Public/Graph/MovieGraphQuickRender.h
@@ -107,7 +107,7 @@ private:
 	 * Returns a new level sequence used for utility purposes (like overriding the camera in use, or changing level
 	 * visibility). It will be injected into the rendering level sequence as a subsequence.
 	 */
-	ULevelSequence* SetUpUtilityLevelSequence() const;
+	ULevelSequence* SetUpUtilityLevelSequence(const EMovieGraphQuickRenderMode QuickRenderMode) const;
 
 	/**
 	 * Populates both the rendering and utility level sequences, and assigns them to the relevant members. Returns true if everything could be set
@@ -116,7 +116,7 @@ private:
 	bool SetUpAllLevelSequences(const EMovieGraphQuickRenderMode QuickRenderMode);
 
 	/** Gets the first world found of the specified type. */
-	UWorld* GetWorldOfType(EWorldType::Type WorldType);
+	UWorld* GetWorldOfType(const EWorldType::Type WorldType) const;
 
 	/**
 	 * Gets the explicitly-set playback range of the level sequence being rendered (set either by a selection in Sequencer, or custom start/end
@@ -133,6 +133,9 @@ private:
 	/** Takes care of post-render tasks, like opening the rendered files. */
 	void HandleJobFinished(const UMovieGraphQuickRenderModeSettings* InQuickRenderSettings, const FMoviePipelineOutputData& InGeneratedOutputData);
 
+	/** Determines if the given viewport look flag is currently active. */
+	bool IsViewportLookFlagActive(const EMovieGraphQuickRenderViewportLookFlags ViewportLookFlag) const;
+
 private:
 	/** The temporary queue that is used by quick render. */
 	UPROPERTY(Transient)
