MovieGraph: The random offsets used during 16->8 bit image quantization now uses a fixed seed, instead of GFrameCounter to make resulting images more deterministic.

diff --git a/Engine/Plugins/MovieScene/MovieRenderPipeline/Source/MovieRenderPipelineCore/Private/MoviePipelineImageQuantization.cpp b/Engine/Plugins/MovieScene/MovieRenderPipeline/Source/MovieRenderPipelineCore/Private/MoviePipelineImageQuantization.cpp
--- a/Engine/Plugins/MovieScene/MovieRenderPipeline/Source/MovieRenderPipelineCore/Private/MoviePipelineImageQuantization.cpp
+++ b/Engine/Plugins/MovieScene/MovieRenderPipeline/Source/MovieRenderPipelineCore/Private/MoviePipelineImageQuantization.cpp
@@ -171,7 +171,7 @@ static TArray<FColor> ConvertLinearTosRGB8bppViaLookupTable(FFloat16Color* InCol
 
 	// This quantization uses a table of float values from 0.0-255.0, and does the proper rounding for each pixel.
 	// The proper rounding is done by adding a random value from 0.0-1.0 instead of 0.5 to each value returned from the table.
-	static FRandomStream RandStream(GFrameCounter);
+	FRandomStream RandStream(31337);
 
 	// FRandomStream is too slow to call 30,000,000 times (4k image) so instead
 	// we pre-calculate a reasonable number of offsets into a table, and then 
