MoviePipeline: Updated Quick Render to render the correct frame ranges when viewing a subsequence within Sequencer (and using the Current Sequence mode or the Use Viewport Camera in Sequence mode, in combination with a Custom or Selection Range frame range).

This allows a user to render in the context of what they're viewing in Sequencer, without having to map their subsequence frame ranges back to the "root" (parent) level sequence.

diff --git a/Engine/Plugins/MovieScene/MovieRenderPipeline/Source/MovieRenderPipelineEditor/Private/Graph/MovieGraphQuickRender.cpp b/Engine/Plugins/MovieScene/MovieRenderPipeline/Source/MovieRenderPipelineEditor/Private/Graph/MovieGraphQuickRender.cpp
--- a/Engine/Plugins/MovieScene/MovieRenderPipeline/Source/MovieRenderPipelineEditor/Private/Graph/MovieGraphQuickRender.cpp
+++ b/Engine/Plugins/MovieScene/MovieRenderPipeline/Source/MovieRenderPipelineEditor/Private/Graph/MovieGraphQuickRender.cpp
@@ -821,20 +821,16 @@ TRange<FFrameNumber> UMovieGraphQuickRenderSubsystem::GetPlaybackRange()
 		const int32 SelectionStartFrame = ULevelSequenceEditorBlueprintLibrary::GetSelectionRangeStart();
 		const int32 SelectionEndFrame = ULevelSequenceEditorBlueprintLibrary::GetSelectionRangeEnd();
 
-		if (SelectionStartFrame != 0)
+		if ((SelectionStartFrame != 0) || (SelectionEndFrame != 0))
 		{
-			StartFrame = FFrameNumber(SelectionStartFrame);
-		}
-
-		if (SelectionEndFrame != 0)
-		{
-			EndFrame = FFrameNumber(SelectionEndFrame);
+			StartFrame = ConvertSubSequenceFrameToRootFrame(SelectionStartFrame);
+			EndFrame = ConvertSubSequenceFrameToRootFrame(SelectionEndFrame);
 		}
 	}
 	else if (QuickRenderModeSettings->FrameRangeType == EMovieGraphQuickRenderFrameRangeType::Custom)
 	{
-		StartFrame = FFrameNumber(QuickRenderModeSettings->CustomStartFrame);
-		EndFrame = FFrameNumber(QuickRenderModeSettings->CustomEndFrame);
+		StartFrame = ConvertSubSequenceFrameToRootFrame(QuickRenderModeSettings->CustomStartFrame);
+		EndFrame = ConvertSubSequenceFrameToRootFrame(QuickRenderModeSettings->CustomEndFrame);
 	}
 	else
 	{
@@ -846,6 +842,58 @@ TRange<FFrameNumber> UMovieGraphQuickRenderSubsystem::GetPlaybackRange()
 	return TRange<FFrameNumber>::Inclusive(StartFrame, EndFrame);
 }
 
+FFrameNumber UMovieGraphQuickRenderSubsystem::ConvertSubSequenceFrameToRootFrame(const int32 FrameNum) const
+{
+	// The "current" level sequence is the root-most level sequence; the "focused" level sequence is the one that's currently visible in Sequencer
+	// (and the level sequence that is having its frame number mapped back to the root)
+	const ULevelSequence* CurrentLevelSequence = ULevelSequenceEditorBlueprintLibrary::GetCurrentLevelSequence();
+	const ULevelSequence* FocusedLevelSequence = ULevelSequenceEditorBlueprintLibrary::GetFocusedLevelSequence();
+
+	// There's nothing to convert if Sequencer is currently viewing the root-level sequence
+	if (CurrentLevelSequence == FocusedLevelSequence)
+	{
+		return FrameNum;
+	}
+
+	// The hierarchy does not include the root, has the outermost subsequence at index 0, and the innermost subsequence at the end
+	const TArray<UMovieSceneSubSection*> SubSequenceHierarchy = ULevelSequenceEditorBlueprintLibrary::GetSubSequenceHierarchy();
+	if (SubSequenceHierarchy.IsEmpty())
+	{
+		// Shouldn't happen if the current and focused level sequences are different
+		return FrameNum;
+	}
+
+	// Map the frame number into tick resolution; the OuterToInnerTransform() appears to need the frame time to be in this format
+	const UMovieSceneSubSection* LastSubSection = SubSequenceHierarchy.Last();
+	const FFrameRate SubSectionDisplayRate = LastSubSection->GetSequence()->GetMovieScene()->GetDisplayRate();
+	const FFrameRate SubSectionTickResolution = LastSubSection->GetSequence()->GetMovieScene()->GetTickResolution();
+	FFrameTime RootFrameTime = ConvertFrameTime(FrameNum, SubSectionDisplayRate, SubSectionTickResolution);
+
+	// Walk the subsequence hierarchy in reverse, applying inverse transforms to get back to the outermost (root) frame time 
+	for (int32 Index = SubSequenceHierarchy.Num() - 1; Index >= 0; --Index)
+	{
+		const UMovieSceneSubSection* SubSection = SubSequenceHierarchy[Index];
+		if (!SubSection)
+		{
+			// GetSubSequenceHierarchy() can return nullptr for sections in some cases
+			UE_LOG(LogMovieRenderPipeline, Error, TEXT("Quick Render: Found an invalid subsequence; rendered frame ranges may be incorrect."));
+			continue;
+		}
+
+		if (TOptional<FFrameTime> TransformedTime = SubSection->OuterToInnerTransform().Inverse().TryTransformTime(RootFrameTime))
+		{
+			RootFrameTime = TransformedTime.GetValue();
+		}
+	}
+
+	// Map the frame number back to the display rate
+	const FFrameRate RootDisplayRate = CurrentLevelSequence->GetMovieScene()->GetDisplayRate();
+	const FFrameRate RootTickResolution = CurrentLevelSequence->GetMovieScene()->GetTickResolution();
+	RootFrameTime = ConvertFrameTime(RootFrameTime, RootTickResolution, RootDisplayRate);
+
+	return RootFrameTime.GetFrame();
+}
+
 void UMovieGraphQuickRenderSubsystem::PerformPreRenderSetup(const EMovieGraphQuickRenderMode QuickRenderMode, UWorld* InEditorWorld)
 {
 	// Adds the specified camera to the Camera Cut track in the given level sequence.
diff --git a/Engine/Plugins/MovieScene/MovieRenderPipeline/Source/MovieRenderPipelineEditor/Public/Graph/MovieGraphQuickRender.h b/Engine/Plugins/MovieScene/MovieRenderPipeline/Source/MovieRenderPipelineEditor/Public/Graph/MovieGraphQuickRender.h
--- a/Engine/Plugins/MovieScene/MovieRenderPipeline/Source/MovieRenderPipelineEditor/Public/Graph/MovieGraphQuickRender.h
+++ b/Engine/Plugins/MovieScene/MovieRenderPipeline/Source/MovieRenderPipelineEditor/Public/Graph/MovieGraphQuickRender.h
@@ -124,6 +124,12 @@ private:
 	 */
 	TRange<FFrameNumber> GetPlaybackRange();
 
+	/**
+	 * Converts a frame number in the focused level sequence to the equivalent in the root level sequence. If the focused and root level sequences
+	 * are the same, this just returns the provided frame number. The "focused" level sequence might be a subsequence, for example.
+	 */
+	FFrameNumber ConvertSubSequenceFrameToRootFrame(const int32 FrameNum) const;
+
 	/** Does any setup needed in the PIE world prior to the render starting. */
 	void PerformPreRenderSetup(const EMovieGraphQuickRenderMode QuickRenderMode, UWorld* InEditorWorld);
 
