MoviePipeline: Added experimental support for encoding mp4 files on Windows.

diff --git a/Engine/Plugins/MovieScene/MovieRenderPipeline/Source/MovieRenderPipelineCore/Private/MoviePipelineTelemetry.cpp b/Engine/Plugins/MovieScene/MovieRenderPipeline/Source/MovieRenderPipelineCore/Private/MoviePipelineTelemetry.cpp
--- a/Engine/Plugins/MovieScene/MovieRenderPipeline/Source/MovieRenderPipelineCore/Private/MoviePipelineTelemetry.cpp
+++ b/Engine/Plugins/MovieScene/MovieRenderPipeline/Source/MovieRenderPipelineCore/Private/MoviePipelineTelemetry.cpp
@@ -1,4 +1,4 @@
-﻿// Copyright Epic Games, Inc. All Rights Reserved.
+// Copyright Epic Games, Inc. All Rights Reserved.
 
 #include "MoviePipelineTelemetry.h"
 
@@ -86,6 +86,7 @@ void FMoviePipelineTelemetry::SendBeginShotRenderTelemetry(UMoviePipelineExecuto
 	EventAttributes.Add(FAnalyticsEventAttribute(TEXT("UsesAudio"), ShotRenderTelemetry.bUsesAudio));
 	EventAttributes.Add(FAnalyticsEventAttribute(TEXT("UsesAvid"), ShotRenderTelemetry.bUsesAvid));
 	EventAttributes.Add(FAnalyticsEventAttribute(TEXT("UsesProRes"), ShotRenderTelemetry.bUsesProRes));
+	EventAttributes.Add(FAnalyticsEventAttribute(TEXT("UsesMP4"), ShotRenderTelemetry.bUsesMP4));
 	EventAttributes.Add(FAnalyticsEventAttribute(TEXT("ResolutionX"), ShotRenderTelemetry.ResolutionX));
 	EventAttributes.Add(FAnalyticsEventAttribute(TEXT("ResolutionY"), ShotRenderTelemetry.ResolutionY));
 	EventAttributes.Add(FAnalyticsEventAttribute(TEXT("HandleFrameCount"), ShotRenderTelemetry.HandleFrameCount));
diff --git a/Engine/Plugins/MovieScene/MovieRenderPipeline/Source/MovieRenderPipelineCore/Public/MoviePipelineTelemetry.h b/Engine/Plugins/MovieScene/MovieRenderPipeline/Source/MovieRenderPipelineCore/Public/MoviePipelineTelemetry.h
--- a/Engine/Plugins/MovieScene/MovieRenderPipeline/Source/MovieRenderPipelineCore/Public/MoviePipelineTelemetry.h
+++ b/Engine/Plugins/MovieScene/MovieRenderPipeline/Source/MovieRenderPipelineCore/Public/MoviePipelineTelemetry.h
@@ -1,4 +1,4 @@
-﻿// Copyright Epic Games, Inc. All Rights Reserved.
+// Copyright Epic Games, Inc. All Rights Reserved.
 
 #pragma once
 
@@ -26,6 +26,7 @@ struct MOVIERENDERPIPELINECORE_API FMoviePipelineShotRenderTelemetry
 	bool bUsesAudio = false;
 	bool bUsesAvid = false;
 	bool bUsesProRes = false;
+	bool bUsesMP4 = false;
 	int32 ResolutionX = 0;
 	int32 ResolutionY = 0;
 	int32 HandleFrameCount = 0;
diff --git a/Engine/Plugins/MovieScene/MovieRenderPipeline/Source/MovieRenderPipelineMP4Encoder/MovieRenderPipelineMP4Encoder.Build.cs b/Engine/Plugins/MovieScene/MovieRenderPipeline/Source/MovieRenderPipelineMP4Encoder/MovieRenderPipelineMP4Encoder.Build.cs
new file mode 100644
--- /dev/null
+++ b/Engine/Plugins/MovieScene/MovieRenderPipeline/Source/MovieRenderPipelineMP4Encoder/MovieRenderPipelineMP4Encoder.Build.cs
@@ -0,0 +1,40 @@
+// Copyright Epic Games, Inc. All Rights Reserved.
+
+namespace UnrealBuildTool.Rules
+{
+	public class MovieRenderPipelineMP4Encoder : ModuleRules
+	{
+		public MovieRenderPipelineMP4Encoder(ReadOnlyTargetRules Target) : base(Target)
+		{
+            PublicDependencyModuleNames.AddRange(				
+				new string[] {
+					"MovieRenderPipelineCore",
+					"MovieRenderPipelineSettings",
+					"OpenColorIO",
+					"SlateCore",
+                });
+
+			PrivateDependencyModuleNames.AddRange(
+				new string[]
+				{
+					"Core",
+                    "CoreUObject",
+                    "Engine",
+                }
+            );
+			
+			if (Target.Platform.IsInGroup(UnrealPlatformGroup.Windows))
+			{
+				PublicDelayLoadDLLs.Add("mfreadwrite.dll");
+				PublicDelayLoadDLLs.Add("mfplat.dll");
+				PublicDelayLoadDLLs.Add("mfuuid.dll");
+
+				PublicSystemLibraries.Add("mfreadwrite.lib");
+				PublicSystemLibraries.Add("mfplat.lib");
+				PublicSystemLibraries.Add("mfuuid.lib");
+
+				PrivateIncludePaths.Add("MovieRenderPipelineMP4Encoder/Private/Windows/");
+			}
+        }
+    }
+}
diff --git a/Engine/Plugins/MovieScene/MovieRenderPipeline/Source/MovieRenderPipelineMP4Encoder/Private/GenericPlatform/MoviePipelineMP4Encoder.h b/Engine/Plugins/MovieScene/MovieRenderPipeline/Source/MovieRenderPipelineMP4Encoder/Private/GenericPlatform/MoviePipelineMP4Encoder.h
new file mode 100644
--- /dev/null
+++ b/Engine/Plugins/MovieScene/MovieRenderPipeline/Source/MovieRenderPipelineMP4Encoder/Private/GenericPlatform/MoviePipelineMP4Encoder.h
@@ -0,0 +1,41 @@
+// Copyright Epic Games, Inc. All Rights Reserved.
+
+#pragma once
+
+#include "CoreMinimal.h"
+#include "Misc/FrameRate.h"
+#include "MoviePipelineMP4EncoderCommon.h"
+#include "MovieRenderPipelineCoreModule.h"
+
+/**
+* Dummy implementation for unsupported platforms which returns failure if used, to avoid
+* doing rendering work and then having no output. This should have a matching pubilc API
+* to the implementation in the other platform folders.
+*/
+class FMoviePipelineMP4Encoder
+{
+public:
+	FMoviePipelineMP4Encoder(const FMoviePipelineMP4EncoderOptions& InOptions)
+	: Options(InOptions)
+	{}
+	
+	~FMoviePipelineMP4Encoder() {}
+
+	bool Initialize()
+	{
+		UE_LOG(LogMovieRenderPipeline, Error, TEXT("MP4 Encoder is unsupported on this platform. Please remove it from your configuration."))
+		return false;
+	}
+	void Finalize() {}
+	bool WriteFrame(const uint8* InFrameData) { return false; }
+	bool WriteAudioSample(const TArrayView<int16>& InAudioSamples) { return false; }
+	const FMoviePipelineMP4EncoderOptions& GetOptions() const { return Options; }
+	bool IsInitialized() const { return false; }
+
+private:
+	/** Configure and initialize the output file */
+	bool InitializeEncoder();
+
+private:
+	FMoviePipelineMP4EncoderOptions Options;
+};
\ No newline at end of file
diff --git a/Engine/Plugins/MovieScene/MovieRenderPipeline/Source/MovieRenderPipelineMP4Encoder/Private/Graph/MovieGraphMP4EncoderNode.cpp b/Engine/Plugins/MovieScene/MovieRenderPipeline/Source/MovieRenderPipelineMP4Encoder/Private/Graph/MovieGraphMP4EncoderNode.cpp
new file mode 100644
--- /dev/null
+++ b/Engine/Plugins/MovieScene/MovieRenderPipeline/Source/MovieRenderPipelineMP4Encoder/Private/Graph/MovieGraphMP4EncoderNode.cpp
@@ -0,0 +1,238 @@
+// Copyright Epic Games, Inc. All Rights Reserved.
+
+#include "Graph/MovieGraphMP4EncoderNode.h"
+
+#include "Graph/MovieGraphBlueprintLibrary.h"
+#include "Graph/MovieGraphConfig.h"
+#include "Graph/MovieGraphOCIOHelper.h"
+#include "Graph/MovieGraphPipeline.h"
+#include "Graph/Nodes/MovieGraphGlobalOutputSettingNode.h"
+#include "ImageWriteTask.h"
+#include "MoviePipelineImageQuantization.h"
+#include "MoviePipelineTelemetry.h"
+#include "MovieRenderPipelineCoreModule.h"
+#include "SampleBuffer.h"
+#include "Styling/AppStyle.h"
+
+
+UMovieGraphMP4EncoderNode::UMovieGraphMP4EncoderNode()
+{
+}
+
+#if WITH_EDITOR
+FText UMovieGraphMP4EncoderNode::GetNodeTitle(const bool bGetDescriptive) const
+{
+	static const FText NodeName = NSLOCTEXT("MovieGraphNodes", "NodeName_MP4", "h264 MP4");
+	return NodeName;
+}
+
+FText UMovieGraphMP4EncoderNode::GetMenuCategory() const
+{
+	return NSLOCTEXT("MovieGraphNodes", "WMFNode_Category", "Output Type");
+}
+
+FText UMovieGraphMP4EncoderNode::GetKeywords() const
+{
+	static const FText Keywords = NSLOCTEXT("MovieGraphNodes", "MP4_Keywords", "mp4 h264 h265 windows mpeg mov movie video");
+	return Keywords;
+}
+
+FLinearColor UMovieGraphMP4EncoderNode::GetNodeTitleColor() const
+{
+	static const FLinearColor NodeColor = FLinearColor(0.047f, 0.654f, 0.537f);
+	return NodeColor;
+}
+
+FSlateIcon UMovieGraphMP4EncoderNode::GetIconAndTint(FLinearColor& OutColor) const
+{
+	static const FSlateIcon Icon = FSlateIcon(FAppStyle::GetAppStyleSetName(), "LevelEditor.Tabs.Cinematics");
+
+	OutColor = FLinearColor::White;
+	return Icon;
+}
+#endif // WITH_EDITOR
+
+TUniquePtr<MovieRenderGraph::IVideoCodecWriter> UMovieGraphMP4EncoderNode::Initialize_GameThread(const FMovieGraphVideoNodeInitializationContext& InInitializationContext)
+{
+	bool bIncludeCDOs = true;
+	constexpr bool bExactMatch = true;
+	UMovieGraphGlobalOutputSettingNode* OutputSetting =
+		InInitializationContext.EvaluatedConfig->GetSettingForBranch<UMovieGraphGlobalOutputSettingNode>(GlobalsPinName, bIncludeCDOs, bExactMatch);
+
+	bIncludeCDOs = false;
+	const UMovieGraphMP4EncoderNode* EvaluatedNode = Cast<UMovieGraphMP4EncoderNode>(
+		InInitializationContext.EvaluatedConfig->GetSettingForBranch(GetClass(), FName(InInitializationContext.PassData->Key.RootBranchName), bIncludeCDOs, bExactMatch));
+	checkf(EvaluatedNode, TEXT("MP4 Encoder node could not be found in the graph in branch [%s]."), *InInitializationContext.PassData->Key.RootBranchName.ToString());
+	
+	const FFrameRate SourceFrameRate = InInitializationContext.Pipeline->GetDataSourceInstance()->GetDisplayRate();
+	const FFrameRate EffectiveFrameRate = UMovieGraphBlueprintLibrary::GetEffectiveFrameRate(OutputSetting, SourceFrameRate);
+	
+	FMoviePipelineMP4EncoderOptions Options;
+	Options.OutputFilename = InInitializationContext.FileName;
+	Options.Width = InInitializationContext.Resolution.X;
+	Options.Height = InInitializationContext.Resolution.Y;
+	Options.FrameRate = EffectiveFrameRate;
+	Options.CommonMaxBitRate = FMath::RoundToInt(MaxBitrateInMbps*1024*1024);
+	Options.CommonMeanBitRate = FMath::RoundToInt(AverageBitrateInMbps*1024*1024);
+	Options.CommonQualityVsSpeed = 100;
+	Options.CommonConstantRateFactor = ConstantRateFactor;
+	Options.EncodingProfile = EncodingProfile;
+	Options.EncodingLevel = EncodingLevel;
+	Options.EncodingRateControl = EncodingRateControl;
+	
+	TUniquePtr<FMP4CodecWriter> NewWriter = MakeUnique<FMP4CodecWriter>();
+	NewWriter->Writer = MakeUnique<FMoviePipelineMP4Encoder>(Options);
+
+	// If OCIO is enabled, don't do additional color conversion.
+	NewWriter->bSkipColorConversions = (EvaluatedNode->bOverride_OCIOConfiguration && EvaluatedNode->OCIOConfiguration.bIsEnabled && InInitializationContext.bAllowOCIO);
+	CachedPipeline = InInitializationContext.Pipeline;
+	
+	return NewWriter;
+}
+
+bool UMovieGraphMP4EncoderNode::Initialize_EncodeThread(MovieRenderGraph::IVideoCodecWriter* InWriter)
+{
+	const FMP4CodecWriter* CodecWriter = static_cast<FMP4CodecWriter*>(InWriter);
+	if(!CodecWriter->Writer->Initialize())
+	{
+		UE_LOG(LogMovieRenderPipeline, Error, TEXT("Failed to initialize Movie Pipeline MP4 writer."));
+		return false;
+	}
+	
+	return true;
+}
+
+void UMovieGraphMP4EncoderNode::WriteFrame_EncodeThread(MovieRenderGraph::IVideoCodecWriter* InWriter, FImagePixelData* InPixelData, TArray<FMovieGraphPassData>&& InCompositePasses, TObjectPtr<UMovieGraphEvaluatedConfig> InEvaluatedConfig, const FString& InBranchName)
+{
+	FMP4CodecWriter* CodecWriter = static_cast<FMP4CodecWriter*>(InWriter);
+	if (!CodecWriter->Writer)
+	{
+		return;
+	}
+
+	bool bIncludeCDOs = false;
+	constexpr bool bExactMatch = true;
+	const UMovieGraphMP4EncoderNode* EvaluatedNode = Cast<UMovieGraphMP4EncoderNode>(
+		InEvaluatedConfig->GetSettingForBranch(GetClass(), FName(InBranchName), bIncludeCDOs, bExactMatch));
+	checkf(EvaluatedNode, TEXT("MP4 Encoder node could not be found in the graph in branch [%s]."), *InBranchName);
+
+	bIncludeCDOs = true;
+	const UMovieGraphGlobalOutputSettingNode* OutputSettingNode =
+		InEvaluatedConfig->GetSettingForBranch<UMovieGraphGlobalOutputSettingNode>(GlobalsPinName, bIncludeCDOs, bExactMatch);
+	
+	const UE::MovieGraph::FMovieGraphSampleState* GraphPayload = InPixelData->GetPayload<UE::MovieGraph::FMovieGraphSampleState>();
+	
+	const bool bConvertToSrgb = !CodecWriter->bSkipColorConversions;
+	const TUniquePtr<FImagePixelData> QuantizedPixelData = UE::MoviePipeline::QuantizeImagePixelDataToBitDepth(InPixelData, 8, nullptr, bConvertToSrgb);
+
+	TArray<FPixelPreProcessor> PixelPreProcessors;
+
+#if WITH_OCIO
+	FMovieGraphOCIOHelper::GenerateOcioPixelPreProcessor(GraphPayload, CachedPipeline.Get(), InEvaluatedConfig, EvaluatedNode->OCIOConfiguration, EvaluatedNode->OCIOContext, PixelPreProcessors);
+#endif
+
+	// Do a quick composite of renders/burn-ins.
+	for (const FMovieGraphPassData& CompositePass : InCompositePasses)
+	{
+		// We don't need to copy the data here (even though it's being passed to a async system) because we already made a unique copy of the
+		// burn in/widget data when we decided to composite it.
+		switch (QuantizedPixelData->GetType())
+		{
+		case EImagePixelType::Color:
+			PixelPreProcessors.Add(TAsyncCompositeImage<FColor>(CompositePass.Value->MoveImageDataToNew()));
+			break;
+		case EImagePixelType::Float16:
+			PixelPreProcessors.Add(TAsyncCompositeImage<FFloat16Color>(CompositePass.Value->MoveImageDataToNew()));
+			break;
+		case EImagePixelType::Float32:
+			PixelPreProcessors.Add(TAsyncCompositeImage<FLinearColor>(CompositePass.Value->MoveImageDataToNew()));
+			break;
+		}
+	}
+
+	// This is done on the current thread for simplicity but the composite itself is parallelized.
+	FImagePixelData* PixelData = QuantizedPixelData.Get();
+	for (const FPixelPreProcessor& PreProcessor : PixelPreProcessors)
+	{
+		// PreProcessors are assumed to be valid.
+		PreProcessor(PixelData);
+	}
+
+	const void* Data = nullptr;
+	int64 DataSize;
+	QuantizedPixelData->GetRawData(Data, DataSize);
+
+	// WriteFrame expects Rec 709 8-bit data.
+	CodecWriter->Writer->WriteFrame((uint8*)Data);
+}
+
+void UMovieGraphMP4EncoderNode::BeginFinalize_EncodeThread(MovieRenderGraph::IVideoCodecWriter* InWriter)
+{
+	// If the writer was not initialized, don't try to finalize anything.
+	FMP4CodecWriter* CodecWriter = static_cast<FMP4CodecWriter*>(InWriter);
+	if (!CodecWriter->Writer)
+	{
+		return;
+	}
+
+	if (!CodecWriter->Writer->IsInitialized())
+	{
+		return;
+	}
+
+	const MoviePipeline::FAudioState& AudioData = CachedPipeline->GetAudioRendererInstance()->GetAudioState();
+
+	for (const TPair<int32, MovieRenderGraph::IVideoCodecWriter::FLightweightSourceData>& SourceData : CodecWriter->LightweightSourceData)
+	{
+		if (!AudioData.FinishedSegments.IsValidIndex(SourceData.Key))
+		{
+			UE_LOG(LogMovieRenderPipeline, Warning, TEXT("Invalid shot index was requested for audio data, skipping audio writes."));
+			continue;
+		}
+
+		// Look up the audio segment for this shot
+		const MoviePipeline::FAudioState::FAudioSegment& AudioSegment = AudioData.FinishedSegments[SourceData.Key];
+
+		// Audio data isn't very sample accurate at this point, so we may have generated slightly more (or less) audio than we expect for
+		// the number of frames, so we're actually going to trim down the view of data we provide to match the number of frames rendered,
+		// to avoid any excess audio after the end of the video.
+		const int32 NumFrames = SourceData.Value.SubmittedFrameCount;
+
+		// ToDo: This is possibly dropping fractions of a sample (ie: 1/48,000th) if the audio sample rate can't be evenly divisible by
+		// the frame rate.
+		const int32 SamplesPerFrame = AudioSegment.SampleRate * CodecWriter->Writer->GetOptions().FrameRate.AsInterval();
+		int32 ExpectedSampleCount = NumFrames * SamplesPerFrame * AudioSegment.NumChannels;
+
+		ExpectedSampleCount = FMath::Min(ExpectedSampleCount, AudioSegment.SegmentData.Num());
+		Audio::TSampleBuffer<int16> SampleBuffer = Audio::TSampleBuffer<int16>(AudioSegment.SegmentData.GetData(), ExpectedSampleCount, AudioSegment.NumChannels, AudioSegment.SampleRate);
+
+		const TArrayView<int16> SampleData = SampleBuffer.GetArrayView();
+		CodecWriter->Writer->WriteAudioSample(SampleData);
+	}
+}
+
+void UMovieGraphMP4EncoderNode::Finalize_EncodeThread(MovieRenderGraph::IVideoCodecWriter* InWriter)
+{
+	// Write to disk
+	const FMP4CodecWriter* CodecWriter = static_cast<FMP4CodecWriter*>(InWriter);
+	if(!CodecWriter->Writer)
+	{
+		return;
+	}
+	CodecWriter->Writer->Finalize();
+}
+
+const TCHAR* UMovieGraphMP4EncoderNode::GetFilenameExtension() const
+{
+	return TEXT("mp4");
+}
+
+bool UMovieGraphMP4EncoderNode::IsAudioSupported() const
+{
+	return true;
+}
+
+void UMovieGraphMP4EncoderNode::UpdateTelemetry(FMoviePipelineShotRenderTelemetry* InTelemetry) const
+{
+	InTelemetry->bUsesMP4 = true;
+}
diff --git a/Engine/Plugins/MovieScene/MovieRenderPipeline/Source/MovieRenderPipelineMP4Encoder/Private/MoviePipelineMP4EncoderCommon.h b/Engine/Plugins/MovieScene/MovieRenderPipeline/Source/MovieRenderPipelineMP4Encoder/Private/MoviePipelineMP4EncoderCommon.h
new file mode 100644
--- /dev/null
+++ b/Engine/Plugins/MovieScene/MovieRenderPipeline/Source/MovieRenderPipelineMP4Encoder/Private/MoviePipelineMP4EncoderCommon.h
@@ -0,0 +1,124 @@
+// Copyright Epic Games, Inc. All Rights Reserved.
+
+#pragma once
+#include "CoreMinimal.h"
+#include "Misc/FrameRate.h"
+#include "MoviePipelineMP4EncoderCommon.generated.h"
+
+// This file contains enums/structures that are common to both the Engine/Editor (for use in the UI) and for platform specific implementations of encoders.
+
+/**
+* Which encoding profile should be used for encoding? A higher profile usually means better quality for a given bitrate,
+* but may not play back on older hardware.
+*/
+UENUM(BlueprintType)
+enum class EMoviePipelineMP4EncodeProfile : uint8
+{
+	Baseline,
+	Main,
+	High,
+};
+
+/**
+* A higher level generally results in a higher quality for a given bitrate, but
+* a higher level requires newer encoders and decoders. Auto will let the encoder
+* choose an appropriate one given the other parameters and is generally the best
+* choice short of external needs.
+*/
+UENUM(BlueprintType)
+enum class EMoviePipelineMP4EncodeLevel : uint8
+{
+	/** Let the encoder choose the best level based on other parameters. */
+	Auto = 0,
+	Level1 = 10,
+	Level1_b = 11,
+	Level1_1 = 11,
+	Level1_2 = 12,
+	Level1_3 = 13,
+	Level2 = 20,
+	Level2_1 = 21,
+	Level2_2 = 22,
+	Level3 = 30,
+	Level3_1 = 31,
+	Level3_2 = 32,
+	Level4 = 40,
+	Level4_1 = 41,
+	Level4_2 = 42,
+	Level5 = 50,
+	Level5_1 = 51,
+	Level5_2 = 52,
+};
+
+UENUM(BlueprintType)
+enum class EMoviePipelineMP4EncodeRateControlMode : uint8
+{
+	/** Automatically chooses a bit rate to target the given quality. Valid ranges are 16-51. */
+	ConstantQP,
+	/** Automatically chooses a bit rate to target the given quality. Valid ranges are 16-51. */
+	Quality,
+	/** Uses a variable bitrate that attempts to achieve a given mean bitrate, but can use a higher bitrate (with no max) on a given frame if needed. Value is in bytes per second.*/
+	VariableBitRate_Constrained,
+	/** Attempts to achieve a given mean bitrate for every frame. Can result in higher bitrates than necessary on simple frames, and lower bitrates than required on complex frames. Value is in bytes per second.*/
+	VariableBitRate_Unconstrained,
+	/** Uses a variable bitrate that attempts to achieve a given mean bitrate, but can specify a maximum bitrate at which point quality will drop if needed. Value is in bytes per second.*/
+	ConstantBitRate,
+}; 
+
+struct FMoviePipelineMP4EncoderOptions
+{
+	FMoviePipelineMP4EncoderOptions()
+		: OutputFilename()
+		, Width(0)
+		, Height(0)
+		, FrameRate(30, 1)
+		, AudioChannelCount(2)
+		, AudioSampleRate(48000)
+		, AudioAverageBitRate(24000) // 192 kbps (bits)
+		, CommonMeanBitRate(12 * 1024 * 1024) // 12 mbps (bits)
+		, CommonMaxBitRate(12 * 1024 * 1024) // 12 mbps (bits)
+		, CommonQualityVsSpeed(100) // 0-33 Low Complexity, 34-66 Medium, 67-100 High (Higher = Slower but Better Quality)
+		, CommonConstantRateFactor(18) // 16-51, Higher is Worse
+		, EncodingProfile(EMoviePipelineMP4EncodeProfile::High)
+		, EncodingLevel(EMoviePipelineMP4EncodeLevel::Auto)
+		, EncodingRateControl(EMoviePipelineMP4EncodeRateControlMode::Quality)
+	{}
+
+	/** The absolute path on disk to try and save the video file to. */
+	FString OutputFilename;
+
+	/** The width of the video file. */
+	uint32 Width;
+
+	/** The height of the video file. */
+	uint32 Height;
+
+	/** Frame Rate of the output video. */
+	FFrameRate FrameRate;
+
+	/** Number of audio channels for the audio track. */
+	uint32 AudioChannelCount;
+
+	/** Number of samples per second (ie: 48'000) for the audio track. */
+	uint32 AudioSampleRate;
+
+	/** Average bitrate for audio track. (in bytes-per-second). Only supported values are 12'000, 16'000, 20'000, and 24'000 */
+	uint32 AudioAverageBitRate;
+
+	/** Average bytes per second for ConstantBitRate, VariableBitRate_Constrained, VariableBitRate_Unconstrained rate control modes. */
+	uint32 CommonMeanBitRate;
+	/** Maximum bytes per second for VariableBitRate_Constrained */
+	uint32 CommonMaxBitRate;
+	/* Quality vs. Speed tradeoff during encode. 0 is faster but worse encode, 100 is slower but higher quality. */
+	uint32 CommonQualityVsSpeed;
+	/* If using Constant Quality, what is the CRF value? 16-51 is the valid range. */
+	uint32 CommonConstantRateFactor;
+
+	/** Which profile should be used when encoding? */
+	EMoviePipelineMP4EncodeProfile EncodingProfile;
+
+	/** What level should the profile use? */
+	EMoviePipelineMP4EncodeLevel EncodingLevel;
+
+	/** Which encoding rate control method should be used? */
+	EMoviePipelineMP4EncodeRateControlMode EncodingRateControl;
+};
\ No newline at end of file
diff --git a/Engine/Plugins/MovieScene/MovieRenderPipeline/Source/MovieRenderPipelineMP4Encoder/Private/MoviePipelineMP4EncoderOutput.cpp b/Engine/Plugins/MovieScene/MovieRenderPipeline/Source/MovieRenderPipelineMP4Encoder/Private/MoviePipelineMP4EncoderOutput.cpp
new file mode 100644
--- /dev/null
+++ b/Engine/Plugins/MovieScene/MovieRenderPipeline/Source/MovieRenderPipelineMP4Encoder/Private/MoviePipelineMP4EncoderOutput.cpp
@@ -0,0 +1,133 @@
+// Copyright Epic Games, Inc. All Rights Reserved.
+
+#include "MoviePipelineMP4EncoderOutput.h"
+#include "ImagePixelData.h"
+#include "MoviePipeline.h"
+#include "MoviePipelineImageQuantization.h"
+#include "MoviePipelineOutputSetting.h"
+#include "MoviePipelinePrimaryConfig.h"
+#include "MoviePipelineTelemetry.h"
+#include "MovieRenderPipelineCoreModule.h"
+#include "SampleBuffer.h"
+
+// For logs
+#include "MovieRenderPipelineCoreModule.h"
+
+TUniquePtr<MovieRenderPipeline::IVideoCodecWriter> UMoviePipelineMP4EncoderOutput::Initialize_GameThread(const FString& InFileName, FIntPoint InResolution, EImagePixelType InPixelType, ERGBFormat InPixelFormat, uint8 InBitDepth, uint8 InNumChannels)
+{
+	const UMoviePipelineOutputSetting* OutputSettings = GetPipeline()->GetPipelinePrimaryConfig()->FindSetting<UMoviePipelineOutputSetting>();
+	if (!OutputSettings)
+	{
+		return nullptr; 
+	} 
+	 
+	FMoviePipelineMP4EncoderOptions Options;
+	Options.OutputFilename = InFileName;
+	Options.Width = InResolution.X;
+	Options.Height = InResolution.Y;
+	Options.FrameRate = GetPipeline()->GetPipelinePrimaryConfig()->GetEffectiveFrameRate(GetPipeline()->GetTargetSequence());
+	Options.CommonMaxBitRate = FMath::RoundToInt(MaxBitrateInMbps*1024*1024);
+	Options.CommonMeanBitRate = FMath::RoundToInt(AverageBitrateInMbps*1024*1024);
+	Options.CommonQualityVsSpeed = 100;
+	Options.CommonConstantRateFactor = ConstantRateFactor;
+	Options.EncodingProfile = EncodingProfile;
+	Options.EncodingLevel = EncodingLevel;
+	Options.EncodingRateControl = EncodingRateControl;
+	
+	TUniquePtr<FMoviePipelineMP4Encoder> Writer = MakeUnique<FMoviePipelineMP4Encoder>(Options);
+	
+	TUniquePtr<FMP4CodecWriter> OutWriter = MakeUnique<FMP4CodecWriter>();
+	OutWriter->Writer = MoveTemp(Writer);
+	OutWriter->FileName = InFileName;
+	
+	return OutWriter;
+}
+
+bool UMoviePipelineMP4EncoderOutput::Initialize_EncodeThread(MovieRenderPipeline::IVideoCodecWriter* InWriter)
+{
+	FMP4CodecWriter* CodecWriter = static_cast<FMP4CodecWriter*>(InWriter);
+	if(!CodecWriter->Writer->Initialize())
+	{
+		UE_LOG(LogMovieRenderPipeline, Error, TEXT("Failed to initialize Media Foundation Sink Writer."));
+		return false;
+	}
+
+	return true;
+}
+
+void UMoviePipelineMP4EncoderOutput::WriteFrame_EncodeThread(MovieRenderPipeline::IVideoCodecWriter* InWriter, FImagePixelData* InPixelData, TArray<MoviePipeline::FCompositePassInfo>&& InCompositePasses)
+{
+	FMP4CodecWriter* CodecWriter = static_cast<FMP4CodecWriter*>(InWriter);
+	
+	// Quantize our 16 bit float data to 8 bit.
+	TUniquePtr<FImagePixelData> QuantizedPixelData = UE::MoviePipeline::QuantizeImagePixelDataToBitDepth(InPixelData, 8);
+	 
+	const void* Data = nullptr;
+	int64 DataSize;
+	QuantizedPixelData->GetRawData(Data, DataSize);
+
+	// WriteFrame expects Rec 709 8-bit data.
+	CodecWriter->Writer->WriteFrame((uint8*)Data);
+}
+
+void UMoviePipelineMP4EncoderOutput::BeginFinalize_EncodeThread(MovieRenderPipeline::IVideoCodecWriter* InWriter)
+{
+	const MoviePipeline::FAudioState& AudioData = GetPipeline()->GetAudioState();
+	FMP4CodecWriter* CodecWriter = static_cast<FMP4CodecWriter*>(InWriter);
+
+	// If the writer was not initialized, don't try to finalize anything.
+	if (!CodecWriter->Writer)
+	{
+		return;
+	}
+
+	if (!CodecWriter->Writer->IsInitialized())
+	{
+		return;
+	}
+
+	for (const TPair<int32, MovieRenderPipeline::IVideoCodecWriter::FLightweightSourceData>& SourceData : CodecWriter->LightweightSourceData)
+	{
+		if (!AudioData.FinishedSegments.IsValidIndex(SourceData.Key))
+		{
+			UE_LOG(LogMovieRenderPipeline, Warning, TEXT("Invalid shot index was requested for audio data, skipping audio writes."));
+			continue;
+		}
+
+		// Look up the audio segment for this shot
+		const MoviePipeline::FAudioState::FAudioSegment& AudioSegment = AudioData.FinishedSegments[SourceData.Key];
+
+		// Audio data isn't very sample accurate at this point, so we may have generated slightly more (or less) audio than we expect for
+		// the number of frames, so we're actually going to trim down the view of data we provide to match the number of frames rendered,
+		// to avoid any excess audio after the end of the video.
+		const int32 NumFrames = SourceData.Value.SubmittedFrameCount;
+
+		// ToDo: This is possibly dropping fractions of a sample (ie: 1/48,000th) if the audio sample rate can't be evenly divisible by
+		// the frame rate.
+		const int32 SamplesPerFrame = AudioSegment.SampleRate * CodecWriter->Writer->GetOptions().FrameRate.AsInterval();
+		int32 ExpectedSampleCount = NumFrames * SamplesPerFrame * AudioSegment.NumChannels;
+
+		ExpectedSampleCount = FMath::Min(ExpectedSampleCount, AudioSegment.SegmentData.Num());
+		Audio::TSampleBuffer<int16> SampleBuffer = Audio::TSampleBuffer<int16>(AudioSegment.SegmentData.GetData(), ExpectedSampleCount, AudioSegment.NumChannels, AudioSegment.SampleRate);
+
+		const TArrayView<int16> SampleData = SampleBuffer.GetArrayView();
+		CodecWriter->Writer->WriteAudioSample(SampleData);
+	}
+}
+
+void UMoviePipelineMP4EncoderOutput::Finalize_EncodeThread(MovieRenderPipeline::IVideoCodecWriter* InWriter)
+{
+	// Commit this to disk.
+	FMP4CodecWriter* CodecWriter = static_cast<FMP4CodecWriter*>(InWriter);
+	if(!CodecWriter->Writer)
+	{
+		return;
+	}
+
+	CodecWriter->Writer->Finalize();
+}
+
+void UMoviePipelineMP4EncoderOutput::UpdateTelemetry(FMoviePipelineShotRenderTelemetry* InTelemetry) const
+{
+	InTelemetry->bUsesMP4 = true;
+}
diff --git a/Engine/Plugins/MovieScene/MovieRenderPipeline/Source/MovieRenderPipelineMP4Encoder/Private/MovieRenderPipelineMP4EncoderModule.cpp b/Engine/Plugins/MovieScene/MovieRenderPipeline/Source/MovieRenderPipelineMP4Encoder/Private/MovieRenderPipelineMP4EncoderModule.cpp
new file mode 100644
--- /dev/null
+++ b/Engine/Plugins/MovieScene/MovieRenderPipeline/Source/MovieRenderPipelineMP4Encoder/Private/MovieRenderPipelineMP4EncoderModule.cpp
@@ -0,0 +1,4 @@
+// Copyright Epic Games, Inc. All Rights Reserved.
+#include "Modules/ModuleManager.h"
+
+IMPLEMENT_MODULE(FDefaultModuleImpl, MovieRenderPipelineMP4Encoder)
diff --git a/Engine/Plugins/MovieScene/MovieRenderPipeline/Source/MovieRenderPipelineMP4Encoder/Private/Windows/MoviePipelineMP4Encoder.cpp b/Engine/Plugins/MovieScene/MovieRenderPipeline/Source/MovieRenderPipelineMP4Encoder/Private/Windows/MoviePipelineMP4Encoder.cpp
new file mode 100644
--- /dev/null
+++ b/Engine/Plugins/MovieScene/MovieRenderPipeline/Source/MovieRenderPipelineMP4Encoder/Private/Windows/MoviePipelineMP4Encoder.cpp
@@ -0,0 +1,829 @@
+
+#include "MoviePipelineMP4Encoder.h"
+
+#if PLATFORM_WINDOWS
+#include "Windows/AllowWindowsPlatformTypes.h"
+#include "Windows/WindowsPlatformMisc.h"
+#include "Windows/WindowsHWrapper.h"
+#endif
+
+THIRD_PARTY_INCLUDES_START
+#include <Windows.h>
+#include <mfapi.h>
+#include <mfidl.h>
+#include <Mfreadwrite.h>
+#include <mferror.h>
+#include <Codecapi.h>
+#include <strmif.h>
+THIRD_PARTY_INCLUDES_END
+
+#if PLATFORM_WINDOWS
+#include "Windows/HideWindowsPlatformTypes.h"
+#endif
+
+
+FMoviePipelineMP4Encoder::FMoviePipelineMP4Encoder(const FMoviePipelineMP4EncoderOptions& InOptions)
+	: Options(InOptions)
+	, bInitialized(false)
+	, bFinalized(false)
+	, NumVideoSamplesWritten(0)
+	, NumAudioSamplesWritten(0)
+	, SinkWriter(nullptr)
+	, VideoStreamIndex(0)
+	, AudioStreamIndex(0)
+{
+
+}
+
+FMoviePipelineMP4Encoder::~FMoviePipelineMP4Encoder()
+{
+	// Insure Finalize is called so that we release the COM library if we were ever initialized.
+	Finalize();
+}
+
+bool FMoviePipelineMP4Encoder::Initialize()
+{
+	// Initialize us for single-threaded communication with the library.
+	if (!FWindowsPlatformMisc::CoInitialize())
+	{
+		UE_LOG(LogTemp, Error, TEXT("Failed to initialize COM library."));
+		return false;
+	}
+
+	// Initialize the Microsoft Media Foundation
+	HRESULT Result = MFStartup(MF_VERSION);
+	if (!SUCCEEDED(Result))
+	{
+		UE_LOG(LogTemp, Error, TEXT("Failed to initialize Microsoft Media Foundation."));
+		return false;
+	}
+
+	bool bResult = InitializeEncoder();
+	if (!bResult)
+	{
+		UE_LOG(LogTemp, Error, TEXT("Failed to initialize Sink Writer."));
+		return false;
+	}
+
+	bInitialized = true;
+	return true;
+}
+
+void FMoviePipelineMP4Encoder::Finalize()
+{
+	if (bFinalized || !bInitialized)
+	{
+		return;
+	}
+
+	if (SinkWriter)
+	{
+		HRESULT Result = SinkWriter->Finalize();
+		if (!SUCCEEDED(Result))
+		{
+			UE_LOG(LogTemp, Error, TEXT("Failed to finalize Sink Writer."));
+		}
+
+		SinkWriter->Release();
+		SinkWriter = nullptr;
+	}
+
+	HRESULT Result = MFShutdown();
+	if (!SUCCEEDED(Result))
+	{
+		UE_LOG(LogTemp, Error, TEXT("Failed to shut down Microsoft Media Foundation."));
+	}
+
+	// Release the COM library
+	FWindowsPlatformMisc::CoUninitialize();
+
+	bFinalized = true;
+}
+
+
+bool FMoviePipelineMP4Encoder::WriteFrame(const uint8* InFrameData)
+{
+	TRACE_CPUPROFILER_EVENT_SCOPE(WMFVideoEncoder_WriteFrame);
+	if (!ensureMsgf(bInitialized || bFinalized, TEXT("WriteFrame should not be called if not initialized or after finalize! Initialized: %d Finalized: %d"), bInitialized, bFinalized))
+	{
+		return false;
+	}
+
+	IMFSample* Sample = nullptr;
+	IMFMediaBuffer* Buffer = nullptr;
+	HRESULT Result;
+
+	const LONG SourceStride = Options.Width * 4; // 4 bytes per pixel
+	const LONG BufferSize = SourceStride * Options.Height;
+	{
+		Result = MFCreateMemoryBuffer(BufferSize, &Buffer);
+
+		if (!SUCCEEDED(Result))
+		{
+			UE_LOG(LogTemp, Error, TEXT("Failed to allocate Media Foundation Buffer to write frame. Width: %d Height: %d Size: %d"), Options.Width, Options.Height, BufferSize);
+			return false;
+		}
+	}
+
+	// Lock the buffer and copy our incoming frame data into the Media Foundation buffer.
+	BYTE* DestinationData = nullptr;
+	{
+		Result = Buffer->Lock(&DestinationData, NULL, NULL);
+	}
+
+	// We can't early out on fail since we need to release the resources in the event of failure.
+	if (SUCCEEDED(Result))
+	{
+		Result = MFCopyImage(DestinationData, SourceStride, (BYTE*)InFrameData, SourceStride, SourceStride, Options.Height);
+
+		// The below snippet would let you read in data that is from bottom-to-top, instead of top-to-bottom. Depending on
+		// the parameters, the hardware and software encoder have produced different ordering results, see comment on
+		// MF_MT_VIDEO_NOMINAL_RANGE below.
+		/*Result = MFCopyImage(DestinationData,
+			SourceStride,
+			((BYTE*) InFrameData) + (Options.Height-1)*SourceStride,
+			-SourceStride,
+			SourceStride,
+			Options.Height); // Not Flipped*/
+	}
+
+	if (SUCCEEDED(Result))
+	{
+		Result = Buffer->Unlock();
+	}
+
+	if (SUCCEEDED(Result))
+	{
+		// Specify how much of the data in the buffer is valid.
+		Result = Buffer->SetCurrentLength(BufferSize);
+	}
+
+	// Create a new sample
+	if (SUCCEEDED(Result))
+	{
+		Result = MFCreateSample(&Sample);
+	}
+
+	if (SUCCEEDED(Result))
+	{
+		Result = Sample->AddBuffer(Buffer);
+	}
+
+	// Duration & Timestamp for this frame.
+	if (SUCCEEDED(Result))
+	{
+		// In 100 nano-second units.
+		uint64 FrameDuration = (10 * 1000 * 1000) * Options.FrameRate.AsInterval();
+		Result = Sample->SetSampleTime(FrameDuration * NumVideoSamplesWritten);
+		NumVideoSamplesWritten++;
+	}
+
+	if (SUCCEEDED(Result))
+	{
+		// In 100 nano-second units.
+		uint64 FrameDuration = (10 * 1000 * 1000) * Options.FrameRate.AsInterval();
+		Result = Sample->SetSampleDuration(FrameDuration);
+	}
+
+	// Send the sample to the Sink Writer
+	if (SUCCEEDED(Result))
+	{
+		Result = SinkWriter->WriteSample(VideoStreamIndex, Sample);
+	}
+
+	if (Sample)
+	{
+		Sample->Release();
+	}
+	if (Buffer)
+	{
+		Buffer->Release();
+	}
+
+	return Result == S_OK;
+}
+
+bool FMoviePipelineMP4Encoder::WriteAudioSample(const TArrayView<int16>& InAudioSamples)
+{
+	TRACE_CPUPROFILER_EVENT_SCOPE(MRQ_WMFSink_WriteAudioSample);
+
+	// If we were not initialized, we can't write the sample.
+	if (!ensureMsgf(bInitialized || bFinalized, TEXT("WriteAudioSample should not be called if not initialized or after finalize! Initialized: %d Finalized: %d"), bInitialized, bFinalized))
+	{
+		return false;
+	}
+
+	// Figure out how many frames worth of audio are being provided by the TArrayView.
+	int32 SampleCountPerChannel = InAudioSamples.Num() / Options.AudioChannelCount;
+	int32 SamplesPerFrame = Options.AudioSampleRate * Options.FrameRate.AsInterval();
+	int32 NumFramesOfData = SampleCountPerChannel / SamplesPerFrame;
+
+	HRESULT Result = S_OK;
+
+	for (int32 SampleIndex = 0; SampleIndex < NumFramesOfData; SampleIndex++)
+	{
+		IMFSample* Sample = nullptr;
+		IMFMediaBuffer* Buffer = nullptr;
+
+		const LONG BufferSize = SamplesPerFrame * sizeof(int16) * Options.AudioChannelCount;
+		Result = MFCreateMemoryBuffer(BufferSize, &Buffer);
+
+		if (!SUCCEEDED(Result))
+		{
+			UE_LOG(LogTemp, Error, TEXT("Failed to allocate Media Foundation Buffer to write audio sample. Buffer Size: %d"), BufferSize);
+			return false;
+		}
+
+		// Lock the buffer and copy our incoming frame data into the Media Foundation buffer.
+		BYTE* DestinationData = nullptr;
+		Result = Buffer->Lock(&DestinationData, NULL, NULL);
+
+		if (SUCCEEDED(Result))
+		{
+			// Figure out the offset into the source data.
+			int16* SourceOffset = (int16*)InAudioSamples.GetData() + (SampleIndex * (SamplesPerFrame * Options.AudioChannelCount));
+
+			// MFCopyImage seems to really just be a generic SSE optimized copy function. 
+			Result = MFCopyImage(DestinationData,
+				BufferSize,
+				(BYTE*)SourceOffset,
+				BufferSize,
+				BufferSize,
+				1);
+		}
+
+		Buffer->Unlock();
+
+		if (SUCCEEDED(Result))
+		{
+			// Specify how much of the data in the buffer is valid.
+			Result = Buffer->SetCurrentLength(BufferSize);
+		}
+
+		// Create a new sample
+		if (SUCCEEDED(Result))
+		{
+			Result = MFCreateSample(&Sample);
+		}
+
+		if (SUCCEEDED(Result))
+		{
+			Result = Sample->AddBuffer(Buffer);
+		}
+
+		// Duration & Timestamp for this frame.
+		if (SUCCEEDED(Result))
+		{
+			// The duration matches the duration of one frame.
+			uint64 FrameDuration = (10 * 1000 * 1000) * Options.FrameRate.AsInterval();
+			Result = Sample->SetSampleDuration(FrameDuration);
+		}
+
+		if (SUCCEEDED(Result))
+		{
+			uint64 FrameDuration = (10 * 1000 * 1000) * Options.FrameRate.AsInterval();
+			Result = Sample->SetSampleTime(NumAudioSamplesWritten * FrameDuration);
+
+			// We track this in the file and not just this for loop because WriteAudioSample can get called
+			// repeatedly in the event that this video rendered multiple separate shots.
+			NumAudioSamplesWritten++;
+		}
+
+		// Send the sample to the Sink Writer
+		if (SUCCEEDED(Result))
+		{
+			Result = SinkWriter->WriteSample(AudioStreamIndex, Sample);
+		}
+
+		if (Sample)
+		{
+			Sample->Release();
+		}
+		if (Buffer)
+		{
+			Buffer->Release();
+		}
+
+		// Early out if we were unable to write a sample for some reason,
+		// we should have freed any buffers allocated during this loop.
+		if (!SUCCEEDED(Result))
+		{
+			return false;
+		}
+	}
+
+
+	return Result == S_OK;
+}
+
+static uint32_t GetEncodingProfile(const EMoviePipelineMP4EncodeProfile InProfile)
+{
+	switch (InProfile)
+	{
+	case EMoviePipelineMP4EncodeProfile::Baseline: return eAVEncH264VProfile_Base;
+	case EMoviePipelineMP4EncodeProfile::Main: return eAVEncH264VProfile_Main;
+	case EMoviePipelineMP4EncodeProfile::High: return eAVEncH264VProfile_High;
+	default:
+		check(false);
+	}
+	return 0;
+}
+
+static uint32_t GetEncodingRateControl(const EMoviePipelineMP4EncodeRateControlMode InProfile)
+{
+	switch (InProfile)
+	{
+	case EMoviePipelineMP4EncodeRateControlMode::ConstantBitRate: return eAVEncCommonRateControlMode_CBR;
+	case EMoviePipelineMP4EncodeRateControlMode::VariableBitRate_Constrained: return eAVEncCommonRateControlMode_PeakConstrainedVBR;
+	case EMoviePipelineMP4EncodeRateControlMode::VariableBitRate_Unconstrained: return eAVEncCommonRateControlMode_UnconstrainedVBR;
+	case EMoviePipelineMP4EncodeRateControlMode::Quality: return eAVEncCommonRateControlMode_Quality;
+	case EMoviePipelineMP4EncodeRateControlMode::ConstantQP: return eAVEncCommonRateControlMode_PeakConstrainedVBR;
+	default:
+		check(false);
+	}
+	return 0;
+}
+
+static HRESULT CreateVideoMediaOutputStream(IMFMediaType** OutVideoMediaTypeOutput, const FMoviePipelineMP4EncoderOptions& InOptions)
+{
+	HRESULT Result = S_OK;
+	if (SUCCEEDED(Result))
+	{
+		Result = MFCreateMediaType(OutVideoMediaTypeOutput);
+	}
+
+	// Set the major media type for the stream.
+ 	if (SUCCEEDED(Result))
+	{
+		Result = (*OutVideoMediaTypeOutput)->SetGUID(MF_MT_MAJOR_TYPE, MFMediaType_Video);
+	}
+
+	// Set the sub-type for the stream
+	if (SUCCEEDED(Result))
+	{
+		Result = (*OutVideoMediaTypeOutput)->SetGUID(MF_MT_SUBTYPE, MFVideoFormat_H264);
+	}
+
+	// Interlace Mode
+	if (SUCCEEDED(Result))
+	{
+		// No support for interleaving.
+		Result = (*OutVideoMediaTypeOutput)->SetUINT32(MF_MT_INTERLACE_MODE, MFVideoInterlace_Progressive);
+	}
+
+	// Specify the Width/Height of the Output
+	if (SUCCEEDED(Result))
+	{
+		Result = MFSetAttributeSize((*OutVideoMediaTypeOutput), MF_MT_FRAME_SIZE, InOptions.Width, InOptions.Height);
+	}
+
+	// Frame Rate
+	if (SUCCEEDED(Result))
+	{
+		Result = MFSetAttributeRatio((*OutVideoMediaTypeOutput), MF_MT_FRAME_RATE, InOptions.FrameRate.Numerator, InOptions.FrameRate.Denominator);
+	}
+
+	// Pixel Aspect Ratio
+	if (SUCCEEDED(Result))
+	{
+		// Always square pixels.
+		Result = MFSetAttributeRatio((*OutVideoMediaTypeOutput), MF_MT_PIXEL_ASPECT_RATIO, 1, 1);
+	}
+
+	// Color Primaries
+	if (SUCCEEDED(Result))
+	{
+		Result = (*OutVideoMediaTypeOutput)->SetUINT32(MF_MT_VIDEO_PRIMARIES, MFVideoPrimaries_BT709);
+	}
+
+	// Color Transfer Function
+	if (SUCCEEDED(Result))
+	{
+		Result = (*OutVideoMediaTypeOutput)->SetUINT32(MF_MT_TRANSFER_FUNCTION, MFVideoTransFunc_sRGB);
+	}
+
+	// Color YUV Matrix
+	if (SUCCEEDED(Result))
+	{
+		Result = (*OutVideoMediaTypeOutput)->SetUINT32(MF_MT_YUV_MATRIX, MFVideoTransferMatrix_BT709);
+	}
+
+	if (SUCCEEDED(Result))
+	{
+		// Setting this to anything other than Unknown seems to flip the resulting output video, but only on
+		// hardware encoders, but not software encoder. But if the software encoder gets used, then the colorspace
+		// is wrong.
+		Result = (*OutVideoMediaTypeOutput)->SetUINT32(MF_MT_VIDEO_NOMINAL_RANGE, MFNominalRange_Unknown);
+	}
+
+	if (SUCCEEDED(Result))
+	{
+		Result = (*OutVideoMediaTypeOutput)->SetUINT32(CODECAPI_AVEncCommonQualityVsSpeed, InOptions.CommonQualityVsSpeed);
+	}
+
+	// Note: Some parameters can be set on the stream before stream initialization, others need to be set as part of the stream
+	// initialization.
+	// These are set for H264 encoding and may be different for other formats in the future.
+	{
+		// There's a number of additional properties supported by the encoder, mostly
+		// related to giving more control over bitrate, GOP, Profiles, etc. Ommitted
+		// for now due to defaults generally being acceptable.
+		// See: https://learn.microsoft.com/en-us/windows/win32/medfound/h-264-video-encoder
+		if (SUCCEEDED(Result))
+		{
+			Result = (*OutVideoMediaTypeOutput)->SetUINT32(MF_MT_MPEG2_PROFILE, GetEncodingProfile(InOptions.EncodingProfile));
+		}
+		if (SUCCEEDED(Result))
+		{
+			EMoviePipelineMP4EncodeLevel Level = InOptions.EncodingLevel;
+			uint32_t LevelValue = Level == EMoviePipelineMP4EncodeLevel::Auto ? -1 : (uint32_t)Level;
+			Result = (*OutVideoMediaTypeOutput)->SetUINT32(MF_MT_MPEG2_LEVEL, LevelValue);
+		}
+
+		if (SUCCEEDED(Result))
+		{
+			Result = (*OutVideoMediaTypeOutput)->SetUINT32(CODECAPI_AVEncH264CABACEnable, true);
+		}
+
+		if (SUCCEEDED(Result))
+		{
+			Result = (*OutVideoMediaTypeOutput)->SetUINT32(CODECAPI_AVEncCommonRateControlMode, GetEncodingRateControl(InOptions.EncodingRateControl));
+		}
+
+		if (SUCCEEDED(Result))
+		{
+			Result = (*OutVideoMediaTypeOutput)->SetUINT32(CODECAPI_AVEncMPVDefaultBPictureCount, 2);
+		}
+	}
+
+	return Result;
+}
+
+static HRESULT CreateAudioMediaOutputStream(IMFMediaType** OutVideoMediaTypeOutput, const FMoviePipelineMP4EncoderOptions& InOptions)
+{
+	HRESULT Result = S_OK;
+	if (SUCCEEDED(Result))
+	{
+		Result = MFCreateMediaType(OutVideoMediaTypeOutput);
+	}
+
+	if (SUCCEEDED(Result))
+	{
+		Result = (*OutVideoMediaTypeOutput)->SetGUID(MF_MT_MAJOR_TYPE, MFMediaType_Audio);
+	}
+
+	// MP4 requires MP3 (or AAC)
+	if (SUCCEEDED(Result))
+	{
+		Result = (*OutVideoMediaTypeOutput)->SetGUID(MF_MT_SUBTYPE, MFAudioFormat_AAC);
+	}
+
+	// Bits per Channel
+	if (SUCCEEDED(Result))
+	{
+		Result = (*OutVideoMediaTypeOutput)->SetUINT32(MF_MT_AUDIO_BITS_PER_SAMPLE, 16);
+	}
+
+	// Samples per Second
+	if (SUCCEEDED(Result))
+	{
+		Result = (*OutVideoMediaTypeOutput)->SetUINT32(MF_MT_AUDIO_SAMPLES_PER_SECOND, InOptions.AudioSampleRate);
+	}
+
+	// Channel Count
+	if (SUCCEEDED(Result))
+	{
+		Result = (*OutVideoMediaTypeOutput)->SetUINT32(MF_MT_AUDIO_NUM_CHANNELS, InOptions.AudioChannelCount);
+	}
+
+	// Average bytes per second
+	if (SUCCEEDED(Result))
+	{
+		Result = (*OutVideoMediaTypeOutput)->SetUINT32(MF_MT_AUDIO_AVG_BYTES_PER_SECOND, InOptions.AudioAverageBitRate);
+	}
+
+	return Result;
+}
+
+static HRESULT CreateVideoMediaTypeIn(IMFMediaType** OutVideoMediaTypeIn, const FMoviePipelineMP4EncoderOptions& InOptions)
+{
+	HRESULT Result = S_OK;
+
+	// Create Video and Audio media types.
+	if (SUCCEEDED(Result))
+	{
+		Result = MFCreateMediaType(OutVideoMediaTypeIn);
+	}
+
+	// Set the major media type for the stream.
+	if (SUCCEEDED(Result))
+	{
+		Result = (*OutVideoMediaTypeIn)->SetGUID(MF_MT_MAJOR_TYPE, MFMediaType_Video);
+	}
+
+	// Set the sub-type for the stream
+	if (SUCCEEDED(Result))
+	{
+		Result = (*OutVideoMediaTypeIn)->SetGUID(MF_MT_SUBTYPE, MFVideoFormat_ARGB32);
+	}
+
+	// Interlace Mode
+	if (SUCCEEDED(Result))
+	{
+		// No support for interleaving.
+		Result = (*OutVideoMediaTypeIn)->SetUINT32(MF_MT_INTERLACE_MODE, MFVideoInterlace_Progressive);
+	}
+
+	// Specify the Width/Height of the Input
+	if (SUCCEEDED(Result))
+	{
+		Result = MFSetAttributeSize((*OutVideoMediaTypeIn), MF_MT_FRAME_SIZE, InOptions.Width, InOptions.Height);
+	}
+
+	// Frame Rate
+	if (SUCCEEDED(Result))
+	{
+		Result = MFSetAttributeRatio((*OutVideoMediaTypeIn), MF_MT_FRAME_RATE, InOptions.FrameRate.Numerator, InOptions.FrameRate.Denominator);
+	}
+
+	// Pixel Aspect Ratio
+	if (SUCCEEDED(Result))
+	{
+		// Always square pixels.
+		Result = MFSetAttributeRatio((*OutVideoMediaTypeIn), MF_MT_PIXEL_ASPECT_RATIO, 1, 1);
+	}
+
+	// Input frames have dependencies on previous samples
+	if (SUCCEEDED(Result))
+	{
+		Result = (*OutVideoMediaTypeIn)->SetUINT32(MF_MT_ALL_SAMPLES_INDEPENDENT, true);
+	}
+
+	return Result;
+}
+
+static HRESULT CreateAudioMediaTypeIn(IMFMediaType** OutAudioMediaTypeIn, const FMoviePipelineMP4EncoderOptions& InOptions)
+{
+	HRESULT Result = S_OK;
+	if (SUCCEEDED(Result))
+	{
+		Result = MFCreateMediaType(OutAudioMediaTypeIn);
+	}
+
+	if (SUCCEEDED(Result))
+	{
+		Result = (*OutAudioMediaTypeIn)->SetGUID(MF_MT_MAJOR_TYPE, MFMediaType_Audio);
+	}
+
+	if (SUCCEEDED(Result))
+	{
+		// MP4 with MP3 requires PCM.
+		Result = (*OutAudioMediaTypeIn)->SetGUID(MF_MT_SUBTYPE, MFAudioFormat_PCM);
+	}
+
+	const int32 BytesPerChannel = 2;
+
+	// Bits Per Sample 
+	if (SUCCEEDED(Result))
+	{
+		// PCM requires signed 16 bit integer. 
+		Result = (*OutAudioMediaTypeIn)->SetUINT32(MF_MT_AUDIO_BITS_PER_SAMPLE, BytesPerChannel * 8);
+	}
+
+	// Input Samples per Second
+	if (SUCCEEDED(Result))
+	{
+		Result = (*OutAudioMediaTypeIn)->SetUINT32(MF_MT_AUDIO_SAMPLES_PER_SECOND, InOptions.AudioSampleRate);
+	}
+
+	// Input Channel Count
+	if (SUCCEEDED(Result))
+	{
+		Result = (*OutAudioMediaTypeIn)->SetUINT32(MF_MT_AUDIO_NUM_CHANNELS, InOptions.AudioChannelCount);
+	}
+
+	// PCM audio has no dependencies on previous samples
+	if (SUCCEEDED(Result))
+	{
+		Result = (*OutAudioMediaTypeIn)->SetUINT32(MF_MT_ALL_SAMPLES_INDEPENDENT, true);
+	}
+
+	// PCM block alignment is numChannels*bytesPerChannel
+	if (SUCCEEDED(Result))
+	{
+		Result = (*OutAudioMediaTypeIn)->SetUINT32(MF_MT_AUDIO_BLOCK_ALIGNMENT, InOptions.AudioChannelCount * BytesPerChannel);
+	}
+
+	// Average byte-rate is just block alignment (above) * samples per second
+	if (SUCCEEDED(Result))
+	{
+		Result = (*OutAudioMediaTypeIn)->SetUINT32(MF_MT_AUDIO_AVG_BYTES_PER_SECOND, (InOptions.AudioChannelCount * BytesPerChannel) * InOptions.AudioSampleRate);
+	}
+
+	return Result;
+}
+
+static HRESULT GetVideoStreamEncoderAttributes(IMFAttributes** OutEncoderAttributes, const FMoviePipelineMP4EncoderOptions& InOptions)
+{
+	HRESULT Result = MFCreateAttributes(OutEncoderAttributes, 1);
+
+	// These are set for H264 encoding and may be different for other formats in the future.
+	{
+		if (InOptions.EncodingRateControl == EMoviePipelineMP4EncodeRateControlMode::ConstantBitRate)
+		{
+			if (SUCCEEDED(Result))
+			{
+				Result = (*OutEncoderAttributes)->SetUINT32(CODECAPI_AVEncCommonMeanBitRate, InOptions.CommonMeanBitRate);
+			}
+
+			if (SUCCEEDED(Result))
+			{
+				Result = (*OutEncoderAttributes)->SetUINT32(CODECAPI_AVEncCommonBufferSize, InOptions.CommonMeanBitRate);
+			}
+
+		}
+		else if (InOptions.EncodingRateControl == EMoviePipelineMP4EncodeRateControlMode::VariableBitRate_Constrained)
+		{
+			if (SUCCEEDED(Result))
+			{
+				Result = (*OutEncoderAttributes)->SetUINT32(CODECAPI_AVEncCommonBufferSize, InOptions.CommonMeanBitRate);
+
+			}
+			if (SUCCEEDED(Result))
+			{
+				Result = (*OutEncoderAttributes)->SetUINT32(CODECAPI_AVEncCommonMaxBitRate, InOptions.CommonMaxBitRate);
+			}
+			if (SUCCEEDED(Result))
+			{
+				Result = (*OutEncoderAttributes)->SetUINT32(CODECAPI_AVEncCommonMeanBitRate, InOptions.CommonMeanBitRate);
+			}
+		}
+		else if (InOptions.EncodingRateControl == EMoviePipelineMP4EncodeRateControlMode::VariableBitRate_Unconstrained)
+		{
+			if (SUCCEEDED(Result))
+			{
+				Result = (*OutEncoderAttributes)->SetUINT32(CODECAPI_AVEncCommonMaxBitRate, InOptions.CommonMaxBitRate);
+			}
+			if (SUCCEEDED(Result))
+			{
+				Result = (*OutEncoderAttributes)->SetUINT32(CODECAPI_AVEncCommonMeanBitRate, InOptions.CommonMeanBitRate);
+			}
+		}
+		else if (InOptions.EncodingRateControl == EMoviePipelineMP4EncodeRateControlMode::Quality)
+		{
+			// This seems to make very similar results to ConstantQP despite being different code.
+			if (SUCCEEDED(Result))
+			{
+				Result = (*OutEncoderAttributes)->SetUINT32(CODECAPI_AVEncCommonRateControlMode, eAVEncCommonRateControlMode_Quality);
+			}
+			if (SUCCEEDED(Result))
+			{
+				Result = (*OutEncoderAttributes)->SetUINT64(CODECAPI_AVEncVideoEncodeQP, InOptions.CommonConstantRateFactor);
+			}
+			if (SUCCEEDED(Result))
+			{
+				Result = (*OutEncoderAttributes)->SetUINT64(CODECAPI_AVEncCommonQuality, 0);
+			}
+		}
+		else if (InOptions.EncodingRateControl == EMoviePipelineMP4EncodeRateControlMode::ConstantQP)
+		{
+			if (SUCCEEDED(Result))
+			{
+				Result = (*OutEncoderAttributes)->SetUINT32(CODECAPI_AVEncCommonQuality, InOptions.CommonConstantRateFactor);
+			}
+			if (SUCCEEDED(Result))
+			{
+				Result = (*OutEncoderAttributes)->SetUINT32(CODECAPI_AVEncVideoMinQP, InOptions.CommonConstantRateFactor);
+			}
+			if (SUCCEEDED(Result))
+			{
+				Result = (*OutEncoderAttributes)->SetUINT32(CODECAPI_AVEncVideoMaxQP, InOptions.CommonConstantRateFactor);
+			}
+		}
+	}
+
+	return Result;
+}
+
+bool FMoviePipelineMP4Encoder::InitializeEncoder()
+{
+	SinkWriter = nullptr;
+
+	// Create a Video and Audio Stream
+	IMFMediaType* VideoMediaTypeOut = nullptr;
+	IMFMediaType* VideoMediaTypeIn = nullptr;
+
+	IMFMediaType* AudioMediaTypeOut = nullptr;
+	IMFMediaType* AudioMediaTypeIn = nullptr;
+
+	IMFAttributes* ConfigAttributes = nullptr;
+	HRESULT Result = MFCreateAttributes(&ConfigAttributes, 1);
+
+	// Avoid any early-outs so that we can still release allocations in the event of an error.
+	if (SUCCEEDED(Result))
+	{
+		// Disable low-latency as we don't need this to be realtime which may affect quality.
+		ConfigAttributes->SetUINT32(CODECAPI_AVLowLatencyMode, false);
+
+		// Use hardware transforms if available. 
+		ConfigAttributes->SetUINT32(MF_READWRITE_ENABLE_HARDWARE_TRANSFORMS, true);
+
+		// Disable throttling in WriteFrame, trading system memory for overall runtime performance.
+		ConfigAttributes->SetUINT32(MF_SINK_WRITER_DISABLE_THROTTLING, true);
+	}
+
+	if (SUCCEEDED(Result))
+	{
+		// Attempt to create the output file.
+		Result = MFCreateSinkWriterFromURL(*Options.OutputFilename, /*Byte Stream Dest*/ NULL, ConfigAttributes, &SinkWriter);
+	}
+	if (SUCCEEDED(Result))
+	{
+		Result = CreateVideoMediaOutputStream(&VideoMediaTypeOut, Options);
+	}
+	// Video Stream
+	if (SUCCEEDED(Result))
+	{
+		Result = SinkWriter->AddStream(VideoMediaTypeOut, (DWORD*)&VideoStreamIndex);
+	}
+
+	if (SUCCEEDED(Result))
+	{
+		Result = CreateAudioMediaOutputStream(&AudioMediaTypeOut, Options);
+	}
+
+	// Audio Stream
+	if (SUCCEEDED(Result))
+	{
+		Result = SinkWriter->AddStream(AudioMediaTypeOut, (DWORD*)&AudioStreamIndex);
+	}
+
+	if (SUCCEEDED(Result))
+	{
+		Result = CreateAudioMediaTypeIn(&AudioMediaTypeIn, Options);
+	}
+
+	if (SUCCEEDED(Result))
+	{
+		Result = CreateVideoMediaTypeIn(&VideoMediaTypeIn, Options);
+	}
+
+	if (SUCCEEDED(Result))
+	{
+		IMFAttributes* pEncAttrs = nullptr;
+		Result = GetVideoStreamEncoderAttributes(&pEncAttrs, Options);
+		if (SUCCEEDED(Result))
+		{
+			Result = SinkWriter->SetInputMediaType(VideoStreamIndex, VideoMediaTypeIn, pEncAttrs);
+		}
+
+		if (pEncAttrs)
+		{
+			pEncAttrs->Release();
+		}
+	}
+
+	if (SUCCEEDED(Result))
+	{
+		Result = SinkWriter->SetInputMediaType(AudioStreamIndex, AudioMediaTypeIn, NULL);
+	}
+
+	ICodecAPI* codecApi;
+	Result = SinkWriter->GetServiceForStream(VideoStreamIndex, GUID_NULL, __uuidof(ICodecAPI), (LPVOID*)&codecApi);
+
+	// The sink writer can now accept data.
+	if (SUCCEEDED(Result))
+	{
+		SinkWriter->BeginWriting();
+	}
+
+	if (SUCCEEDED(Result))
+	{
+		SinkWriter->AddRef();
+	}
+
+	// Release temporary resources.
+	if (SinkWriter)
+	{
+		SinkWriter->Release();
+	}
+
+	if (VideoMediaTypeOut)
+	{
+		VideoMediaTypeOut->Release();
+	}
+	if (VideoMediaTypeIn)
+	{
+		VideoMediaTypeIn->Release();
+	}
+	if (AudioMediaTypeOut)
+	{
+		AudioMediaTypeOut->Release();
+	}
+	if (AudioMediaTypeIn)
+	{
+		AudioMediaTypeIn->Release();
+	}
+
+	return Result == S_OK;
+}
diff --git a/Engine/Plugins/MovieScene/MovieRenderPipeline/Source/MovieRenderPipelineMP4Encoder/Private/Windows/MoviePipelineMP4Encoder.h b/Engine/Plugins/MovieScene/MovieRenderPipeline/Source/MovieRenderPipelineMP4Encoder/Private/Windows/MoviePipelineMP4Encoder.h
new file mode 100644
--- /dev/null
+++ b/Engine/Plugins/MovieScene/MovieRenderPipeline/Source/MovieRenderPipelineMP4Encoder/Private/Windows/MoviePipelineMP4Encoder.h
@@ -0,0 +1,63 @@
+// Copyright Epic Games, Inc. All Rights Reserved.
+
+#pragma once
+
+#include "CoreMinimal.h"
+#include "Misc/FrameRate.h"
+#include "MoviePipelineMP4EncoderCommon.h"
+
+/**
+* Takes in 8-bit RGBA frames which have had an sRGB gamut applied, and outputs a YUV 4:2:0 encoded video.
+* Targeted at non-realtime applications (ie: movie rendering) and not realtime encoding. Experimental.
+* 
+* This should have a matching public API to the implementation in the other platform folders.
+*/
+class FMoviePipelineMP4Encoder
+{
+public:
+	FMoviePipelineMP4Encoder(const FMoviePipelineMP4EncoderOptions& InOptions);
+	~FMoviePipelineMP4Encoder();
+
+	/** Call to initialize the Sink Writer. This must be done before attempting to write data to it. */
+	bool Initialize();
+
+	/** Finalize the video file and finish writing it to disk. Called by the destructor if not automatically called. */
+	void Finalize();
+
+	/** Appends a new frame onto the output file. */
+	bool WriteFrame(const uint8* InFrameData);
+	/** Appends a new audio sample onto the audio stream. */
+	bool WriteAudioSample(const TArrayView<int16>& InAudioSamples);
+
+	const FMoviePipelineMP4EncoderOptions& GetOptions() const { return Options; }
+
+	bool IsInitialized() const { return bInitialized; }
+
+private:
+	/** Configure and initialize the output file */
+	bool InitializeEncoder();
+
+private:
+	/** Input/Output options this writer was initialized with. */
+	FMoviePipelineMP4EncoderOptions Options;
+
+	/** Has Initialize been successfully called? */
+	bool bInitialized;
+
+	/** Has Finalize been called? */
+	bool bFinalized;
+
+	/** How many video samples (frames) have we written so far? */
+	uint64 NumVideoSamplesWritten;
+
+	/** How many audio samples (frames) have we written so far? */
+	uint64 NumAudioSamplesWritten;
+
+	/** The sink writer we are writing samples to. */
+	struct IMFSinkWriter* SinkWriter;
+
+	/** Stream index for video within the Sink Writer */
+	uint32 VideoStreamIndex;
+	/** Stream index for audio within the Sink Writer */
+	uint32 AudioStreamIndex;
+};
\ No newline at end of file
diff --git a/Engine/Plugins/MovieScene/MovieRenderPipeline/Source/MovieRenderPipelineMP4Encoder/Public/Graph/MovieGraphMP4EncoderNode.h b/Engine/Plugins/MovieScene/MovieRenderPipeline/Source/MovieRenderPipelineMP4Encoder/Public/Graph/MovieGraphMP4EncoderNode.h
new file mode 100644
--- /dev/null
+++ b/Engine/Plugins/MovieScene/MovieRenderPipeline/Source/MovieRenderPipelineMP4Encoder/Public/Graph/MovieGraphMP4EncoderNode.h
@@ -0,0 +1,145 @@
+// Copyright Epic Games, Inc. All Rights Reserved.
+
+#pragma once
+
+#include "MoviePipelineMP4Encoder.h"
+#include "Graph/MovieGraphNode.h"
+#include "Graph/Nodes/MovieGraphVideoOutputNode.h"
+
+#include "MovieGraphMP4EncoderNode.generated.h"
+
+/** A node which can output H264 mp4 files. */
+UCLASS(BlueprintType)
+class UMovieGraphMP4EncoderNode : public UMovieGraphVideoOutputNode
+{
+	GENERATED_BODY()
+
+public:
+	UMovieGraphMP4EncoderNode();
+
+#if WITH_EDITOR
+	virtual FText GetNodeTitle(const bool bGetDescriptive = false) const override;
+	virtual FText GetMenuCategory() const override;
+	virtual FText GetKeywords() const override;
+	virtual FLinearColor GetNodeTitleColor() const override;
+	virtual FSlateIcon GetIconAndTint(FLinearColor& OutColor) const override;
+#endif
+
+protected:
+	// UMovieGraphVideoOutputNode Interface
+	virtual TUniquePtr<MovieRenderGraph::IVideoCodecWriter> Initialize_GameThread(const FMovieGraphVideoNodeInitializationContext& InInitializationContext) override;
+	virtual bool Initialize_EncodeThread(MovieRenderGraph::IVideoCodecWriter* InWriter) override;
+	virtual void WriteFrame_EncodeThread(MovieRenderGraph::IVideoCodecWriter* InWriter, FImagePixelData* InPixelData, TArray<FMovieGraphPassData>&& InCompositePasses, TObjectPtr<UMovieGraphEvaluatedConfig> InEvaluatedConfig, const FString& InBranchName) override;
+	virtual void BeginFinalize_EncodeThread(MovieRenderGraph::IVideoCodecWriter* InWriter) override;
+	virtual void Finalize_EncodeThread(MovieRenderGraph::IVideoCodecWriter* InWriter) override;
+	virtual const TCHAR* GetFilenameExtension() const override;
+	virtual bool IsAudioSupported() const override;
+	// ~UMovieGraphVideoOutputNode Interface
+
+	// UMovieGraphSettingNode Interface
+	virtual void UpdateTelemetry(FMoviePipelineShotRenderTelemetry* InTelemetry) const override;
+	// ~UMovieGraphSettingNode Interface
+
+protected:
+	struct FMP4CodecWriter : public MovieRenderGraph::IVideoCodecWriter
+	{
+		bool bSkipColorConversions = false;
+		TUniquePtr<FMoviePipelineMP4Encoder> Writer;
+	};
+	
+	/** The pipeline that is running this node. */
+	TWeakObjectPtr<UMovieGraphPipeline> CachedPipeline;
+
+public:
+	UPROPERTY(EditAnywhere, BlueprintReadWrite, Category = Overrides, meta = (InlineEditConditionToggle))
+	uint8 bOverride_AverageBitrateInMbps : 1;
+	
+	UPROPERTY(EditAnywhere, BlueprintReadWrite, Category = Overrides, meta = (InlineEditConditionToggle))
+	uint8 bOverride_MaxBitrateInMbps : 1;
+	
+	UPROPERTY(EditAnywhere, BlueprintReadWrite, Category = Overrides, meta = (InlineEditConditionToggle))
+	uint8 bOverride_ConstantRateFactor : 1;
+	
+	UPROPERTY(EditAnywhere, BlueprintReadWrite, Category = Overrides, meta = (InlineEditConditionToggle))
+	uint8 bOverride_EncodingProfile : 1;
+	
+	UPROPERTY(EditAnywhere, BlueprintReadWrite, Category = Overrides, meta = (InlineEditConditionToggle))
+	uint8 bOverride_EncodingLevel : 1;
+
+	UPROPERTY(EditAnywhere, BlueprintReadWrite, Category = Overrides, meta = (InlineEditConditionToggle))
+	uint8 bOverride_OCIOConfiguration : 1;
+
+	UPROPERTY(EditAnywhere, BlueprintReadWrite, Category = Overrides, meta = (InlineEditConditionToggle))
+	uint8 bOverride_OCIOContext : 1;
+		
+	/**
+	* Specifies the bitrate control method used by the encoder. Quality lets the user target a given quality without concern
+	* to the filesize, while Average/Constant modes allow you to suggest desired sizes, though the resulting file may still
+	* end up larger or smaller than specified.
+	*/
+	UPROPERTY(BlueprintReadWrite, EditAnywhere, Category = "Settings")
+	EMoviePipelineMP4EncodeRateControlMode EncodingRateControl = EMoviePipelineMP4EncodeRateControlMode::Quality;
+	
+	/**
+	* What is the average bitrate the encoder should target per second? Value is in Megabits per Second,
+	* so a value of 8 will become 8192Kbps (kilobits per second). Higher resolutions and higher framerates
+	* need higher bitrates, reasonable starting values are 8 for 1080p30, 45 for 4k. 
+	* 
+	* Only applies to encoding modes not related to Quality.
+	*/
+	UPROPERTY(BlueprintReadWrite, EditAnywhere, Category = "Settings", meta=(ClampMin=0.1, UIMin=0.1, UIMax=50, EditCondition="bOverride_AverageBitrateInMbps"))
+	float AverageBitrateInMbps = 8.f;
+	
+	/**
+	* When using VariableBitRate_Constrained, what is the maximum bitrate that the encoder can briefly use for
+	* more complex scenes, while still trying to maintain the average set in AverageBitrateInMbps. In theory the
+	* maximum should be twice the average, but often in practice a smaller difference of 3-6Mbps is sufficient.
+	*
+	* Only applies to constrained max bit rate.
+	*/
+	UPROPERTY(BlueprintReadWrite, EditAnywhere, Category = "Settings", meta = (ClampMin = 0.1, UIMin = 0.1, UIMax = 50, EditCondition = "bOverride_MaxBitrateInMbps"))
+	float MaxBitrateInMbps = 16.f;
+
+	/**
+	* What is the Constant Rate Factor (CRF) when targeting a specific quality. Values of 17-18 are generally considered
+	* perceptually lossless, while higher values produce smaller files with lower quality results. There is an absolute
+	* maximum value range of [16-51], where 16 is the highest quality possible and 51 is the lowest quality. This scale
+	* is logarithmic, so small changes can result in large differences in quality, and filesize cost. 
+	*
+	* Only applies to Quality encoding rate control mode.
+	*/
+	UPROPERTY(BlueprintReadWrite, EditAnywhere, Category = "Settings", meta=(ClampMin=16, UIMin=16, ClampMax=51, UIMax=51, EditCondition="bOverride_ConstantRateFactor"))
+	int32 ConstantRateFactor = 20;
+
+	/** A higher profile generally results in a better quality video for the same bitrate, but may not be supported for playback on old devices. */
+	UPROPERTY(BlueprintReadWrite, EditAnywhere, Category = "Settings", meta=(AdvancedDisplay, EditCondition="bOverride_EncodingProfile"))
+	EMoviePipelineMP4EncodeProfile EncodingProfile = EMoviePipelineMP4EncodeProfile::High;
+
+	/** 
+	* A higher encode level generally results in a better quality video for the same bitrate, but may not be supported for playback on old devices.
+	* Auto will let the encoder automatically select the best level given the resolution, profile, and bitrate.
+	*/
+	UPROPERTY(BlueprintReadWrite, EditAnywhere, Category = "Settings", meta=(AdvancedDisplay, EditCondition="bOverride_EncodingLevel"))
+	EMoviePipelineMP4EncodeLevel EncodingLevel = EMoviePipelineMP4EncodeLevel::Auto;
+	
+	/**
+	* OCIO configuration/transform settings.
+	*
+	* Note: There are differences from the previous implementation in MRQ given that we are now doing CPU-side processing.
+	* 1) This feature only works on desktop platforms when the OpenColorIO library is available.
+	* 2) Users are now responsible for setting the renderer output space to Final Color (HDR) in Linear Working Color Space (SCS_FinalColorHDR) by
+	*    disabling the Tone Curve setting on the renderer node.
+	*/
+	UPROPERTY(EditAnywhere, BlueprintReadWrite, Category = "OCIO", DisplayName="OCIO Configuration", meta = (DisplayAfter = "FileNameFormat", EditCondition = "bOverride_OCIOConfiguration"))
+	FOpenColorIODisplayConfiguration OCIOConfiguration;
+
+	/**
+	* OCIO context of key-value string pairs, typically used to apply shot-specific looks (such as a CDL color correction, or a 1D grade LUT).
+	* 
+	* Notes:
+	* 1) If a configuration asset base context was set, it remains active but can be overridden here with new key-values.
+	* 2) Format tokens such as {shot_name} are supported and will get resolved before submission.
+	*/
+	UPROPERTY(EditAnywhere, BlueprintReadWrite, Category = "OCIO", DisplayName = "OCIO Context", meta = (DisplayAfter = "OCIOConfiguration", EditCondition = "bOverride_OCIOContext"))
+	TMap<FString, FString> OCIOContext;
+};
\ No newline at end of file
diff --git a/Engine/Plugins/MovieScene/MovieRenderPipeline/Source/MovieRenderPipelineMP4Encoder/Public/MoviePipelineMP4EncoderOutput.h b/Engine/Plugins/MovieScene/MovieRenderPipeline/Source/MovieRenderPipelineMP4Encoder/Public/MoviePipelineMP4EncoderOutput.h
new file mode 100644
--- /dev/null
+++ b/Engine/Plugins/MovieScene/MovieRenderPipeline/Source/MovieRenderPipelineMP4Encoder/Public/MoviePipelineMP4EncoderOutput.h
@@ -0,0 +1,92 @@
+// Copyright Epic Games, Inc. All Rights Reserved.
+
+#include "MoviePipelineVideoOutputBase.h"
+#include "MoviePipelineMP4Encoder.h"
+#include "MoviePipelineMP4EncoderOutput.generated.h"
+
+// Forward Declare
+class FMoviePipelineMP4Encoder;
+
+UCLASS(BlueprintType)
+class UMoviePipelineMP4EncoderOutput : public UMoviePipelineVideoOutputBase
+{
+	GENERATED_BODY()
+
+		UMoviePipelineMP4EncoderOutput()
+		: UMoviePipelineVideoOutputBase()
+	{
+	}
+
+protected:
+	// UMoviePipelineVideoOutputBase Interface
+	virtual TUniquePtr<MovieRenderPipeline::IVideoCodecWriter> Initialize_GameThread(const FString& InFileName, FIntPoint InResolution, EImagePixelType InPixelType, ERGBFormat InPixelFormat, uint8 InBitDepth, uint8 InNumChannels) override;
+	virtual bool Initialize_EncodeThread(MovieRenderPipeline::IVideoCodecWriter* InWriter) override;
+	virtual void WriteFrame_EncodeThread(MovieRenderPipeline::IVideoCodecWriter* InWriter, FImagePixelData* InPixelData, TArray<MoviePipeline::FCompositePassInfo>&& InCompositePasses) override;
+	virtual void BeginFinalize_EncodeThread(MovieRenderPipeline::IVideoCodecWriter* InWriter);
+	virtual void Finalize_EncodeThread(MovieRenderPipeline::IVideoCodecWriter* InWriter);
+	virtual const TCHAR* GetFilenameExtension() const override { return TEXT("mp4"); }
+	virtual bool IsAudioSupported() const override { return false; }
+	// ~UMoviePipelineVideoOutputBase Interface
+	
+	// UMoviePipelineSetting Interface
+	virtual void UpdateTelemetry(FMoviePipelineShotRenderTelemetry* InTelemetry) const override;
+	// ~UMoviePipelineSetting Interface
+
+	// UMoviePipelineOutputBase Interface
+#if WITH_EDITOR
+	virtual FText GetDisplayText() const override { return NSLOCTEXT("MovieRenderPipeline", "MP4EncoderNode_DisplayName", "h264 MP4 [8bit]"); }
+#endif
+	// ~UMoviePipelineOutputBase Interface
+
+public:
+	/**
+	* Specifies the bitrate control method used by the encoder. Quality lets the user target a given quality without concern
+	* to the filesize, while Average/Constant modes allow you to suggest desired sizes, though the resulting file may still
+	* end up larger or smaller than specified.
+	*/
+	UPROPERTY(BlueprintReadWrite, EditAnywhere, Category = "Settings")
+	EMoviePipelineMP4EncodeRateControlMode EncodingRateControl = EMoviePipelineMP4EncodeRateControlMode::Quality;
+	
+	/**
+	* What is the average bitrate the encoder should target per second? Value is in Megabits per Second,
+	* so a value of 8 will become 8192Kbps (kilobits per second). Higher resolutions and higher framerates
+	* need higher bitrates, reasonable starting values are 8 for 1080p30, 45 for 4k. Only applies to
+	* encoding modes not related to Quality.
+	*/
+	UPROPERTY(BlueprintReadWrite, EditAnywhere, Category = "Settings", meta=(ClampMin=0.1, UIMin=0.1, UIMax=50, EditCondition="EncodingRateControl != EMoviePipelineMP4EncodingRateControlMode::Quality"))
+	float AverageBitrateInMbps = 8.f;
+	
+	/**
+	* When using VariableBitRate_Constrained, what is the maximum bitrate that the encoder can briefly use for
+	* more complex scenes, while still trying to maintain the average set in AverageBitrateInMbps. In theory the
+	* maximum should be twice the average, but often in practice a smaller difference of 3-6Mbps is sufficient.
+	*/
+	UPROPERTY(BlueprintReadWrite, EditAnywhere, Category = "Settings", meta = (ClampMin = 0.1, UIMin = 0.1, UIMax = 50, EditCondition = "EncodingRateControl == EMoviePipelineMP4EncodingRateControlMode::VariableBitRate_Constrained"))
+	float MaxBitrateInMbps = 16.f;
+
+	/**
+	* What is the Constant Rate Factor (CRF) when targeting a specific quality. Values of 17-18 are generally considered
+	* perceptually lossless, while higher values produce smaller files with lower quality results. There is an absolute
+	* maximum value range of [16-51], where 16 is the highest quality possible and 51 is the lowest quality. This scale
+	* is logarithmic, so small changes can result in large differences in quality, and filesize cost.
+	*/
+	UPROPERTY(BlueprintReadWrite, EditAnywhere, Category = "Settings", meta=(ClampMin=16, UIMin=16, ClampMax=51, UIMax=51, EditCondition="EncodingRateControl == EMoviePipelineMP4EncodingRateControlMode::Quality"))
+	int32 ConstantRateFactor = 20;
+
+	/** A higher profile generally results in a better quality video for the same bitrate, but may not be supported for playback on old devices. */
+	UPROPERTY(BlueprintReadWrite, EditAnywhere, Category = "Settings", meta=(AdvancedDisplay))
+	EMoviePipelineMP4EncodeProfile EncodingProfile = EMoviePipelineMP4EncodeProfile::High;
+
+	/** 
+	* A higher encode level generally results in a better quality video for the same bitrate, but may not be supported for playback on old devices.
+	* Auto will let the encoder automatically select the best level given the resolution, profile, and bitrate.
+	*/
+	UPROPERTY(BlueprintReadWrite, EditAnywhere, Category = "Settings", meta=(AdvancedDisplay))
+	EMoviePipelineMP4EncodeLevel EncodingLevel = EMoviePipelineMP4EncodeLevel::Auto;
+
+protected:
+	struct FMP4CodecWriter : public MovieRenderPipeline::IVideoCodecWriter
+	{
+		TUniquePtr<FMoviePipelineMP4Encoder> Writer;
+	};
+};
\ No newline at end of file
