MoviePipeline: Fixed a Quick Render crash that could occur when using the "Selected Cameras" mode in conjunction with the "Selection Range" frame range option.

The crash was due to trying to render a frame range of [0, 0), which is of course not valid. Also did some refactoring to move UI state into its own struct. This makes it easier to access outside of the menu, which is important because it's needed in the customization (which now disables the frame range properties unless a mode which renders more than one frame is active).

diff --git a/Engine/Plugins/MovieScene/MovieRenderPipeline/Source/MovieRenderPipelineCore/Public/Graph/MovieGraphQuickRenderSettings.h b/Engine/Plugins/MovieScene/MovieRenderPipeline/Source/MovieRenderPipelineCore/Public/Graph/MovieGraphQuickRenderSettings.h
--- a/Engine/Plugins/MovieScene/MovieRenderPipeline/Source/MovieRenderPipelineCore/Public/Graph/MovieGraphQuickRenderSettings.h
+++ b/Engine/Plugins/MovieScene/MovieRenderPipeline/Source/MovieRenderPipelineCore/Public/Graph/MovieGraphQuickRenderSettings.h
@@ -176,7 +176,7 @@ public:
 	UPROPERTY(EditAnywhere, BlueprintReadWrite, Category="Quick Render", DisplayName="Apply Viewport Look", meta=(EditCondition="bOverride_ViewportLookFlags", Bitmask, BitmaskEnum="/Script/MovieRenderPipelineCore.EMovieGraphQuickRenderViewportLookFlags"))
 	int32 ViewportLookFlags = static_cast<int32>(EMovieGraphQuickRenderViewportLookFlags::None);
 
-	/** How to determine what the start/end frame should be. */
+	/** The frame range that should be used to render from. Only available in modes that render more than one frame. */
 	UPROPERTY(EditAnywhere, BlueprintReadWrite, Category="Quick Render", DisplayName="Sequencer Frame Range")
 	EMovieGraphQuickRenderFrameRangeType FrameRangeType = EMovieGraphQuickRenderFrameRangeType::PlaybackRange;
 
diff --git a/Engine/Plugins/MovieScene/MovieRenderPipeline/Source/MovieRenderPipelineEditor/Private/Customizations/Graph/MovieGraphQuickRenderSettingsCustomization.h b/Engine/Plugins/MovieScene/MovieRenderPipeline/Source/MovieRenderPipelineEditor/Private/Customizations/Graph/MovieGraphQuickRenderSettingsCustomization.h
--- a/Engine/Plugins/MovieScene/MovieRenderPipeline/Source/MovieRenderPipelineEditor/Private/Customizations/Graph/MovieGraphQuickRenderSettingsCustomization.h
+++ b/Engine/Plugins/MovieScene/MovieRenderPipeline/Source/MovieRenderPipelineEditor/Private/Customizations/Graph/MovieGraphQuickRenderSettingsCustomization.h
@@ -9,6 +9,7 @@
 #include "IDetailGroup.h"
 #include "MovieGraphCustomizationUtils.h"
 #include "MoviePipelineQueue.h"
+#include "MoviePipelineQuickRenderUIState.h"
 #include "MoviePipelineUtils.h"
 #include "PropertyHandle.h"
 #include "UObject/Package.h"
@@ -63,6 +64,32 @@ protected:
 		DetailBuilder->ForceRefreshDetails();
 	}
 
+	void UpdateFrameRangePropertyEnableState() const
+	{
+		const TAttribute<bool> IsEnabledAttr = TAttribute<bool>::CreateLambda([]() -> bool
+		{
+			const EMovieGraphQuickRenderMode CurrentWindowRenderMode = FMoviePipelineQuickRenderUIState::GetWindowQuickRenderMode();
+			
+			// Enabling the frame range properties only makes sense for modes that can render more than one frame.
+			const bool bNewEnableState =
+				(CurrentWindowRenderMode == EMovieGraphQuickRenderMode::CurrentSequence)
+				|| (CurrentWindowRenderMode == EMovieGraphQuickRenderMode::UseViewportCameraInSequence);
+
+			return bNewEnableState;
+		});
+
+		const TSharedRef<IPropertyHandle> FrameRangeTypeProperty =
+			DetailBuilder->GetProperty(GET_MEMBER_NAME_CHECKED(UMovieGraphQuickRenderModeSettings, FrameRangeType), UMovieGraphQuickRenderModeSettings::StaticClass());
+		const TSharedRef<IPropertyHandle> CustomStartFrameProperty =
+			DetailBuilder->GetProperty(GET_MEMBER_NAME_CHECKED(UMovieGraphQuickRenderModeSettings, CustomStartFrame), UMovieGraphQuickRenderModeSettings::StaticClass());
+		const TSharedRef<IPropertyHandle> CustomEndFrameProperty =
+			DetailBuilder->GetProperty(GET_MEMBER_NAME_CHECKED(UMovieGraphQuickRenderModeSettings, CustomEndFrame), UMovieGraphQuickRenderModeSettings::StaticClass());
+
+		DetailBuilder->EditDefaultProperty(FrameRangeTypeProperty)->IsEnabled(IsEnabledAttr);
+		DetailBuilder->EditDefaultProperty(CustomStartFrameProperty)->IsEnabled(IsEnabledAttr);
+		DetailBuilder->EditDefaultProperty(CustomEndFrameProperty)->IsEnabled(IsEnabledAttr);
+	}
+
 	virtual void CustomizeDetails(IDetailLayoutBuilder& InDetailBuilder) override
 	{
 		// Refresh the customization every time a save happens. Use this opportunity to update the variables in the UI. We could update the UI before
@@ -88,6 +115,9 @@ protected:
 			InDetailBuilder.GetProperty(TEXT("GraphVariableAssignments"), UMovieGraphQuickRenderModeSettings::StaticClass());
 		InDetailBuilder.HideProperty(GraphVariableAssignmentsProperty);
 
+		// Update the enable state of the frame range properties to react to changes in the quick render mode setting
+		UpdateFrameRangePropertyEnableState();
+
 		// Refresh the UI if the graph preset changes (so the new variable assignments are displayed).
 		QuickRenderModeSettings->OnGraphChangedDelegate.AddSP(this, &FMovieGraphQuickRenderSettingsCustomization::RefreshLayout);
 
diff --git a/Engine/Plugins/MovieScene/MovieRenderPipeline/Source/MovieRenderPipelineEditor/Private/Graph/MovieGraphQuickRender.cpp b/Engine/Plugins/MovieScene/MovieRenderPipeline/Source/MovieRenderPipelineEditor/Private/Graph/MovieGraphQuickRender.cpp
--- a/Engine/Plugins/MovieScene/MovieRenderPipeline/Source/MovieRenderPipelineEditor/Private/Graph/MovieGraphQuickRender.cpp
+++ b/Engine/Plugins/MovieScene/MovieRenderPipeline/Source/MovieRenderPipelineEditor/Private/Graph/MovieGraphQuickRender.cpp
@@ -79,7 +79,7 @@ private:
 	int32 TransactionId = 0;
 };
 
-void UMovieGraphQuickRenderSubsystem::BeginQuickRender(const EMovieGraphQuickRenderMode QuickRenderMode, const UMovieGraphQuickRenderModeSettings* InQuickRenderSettings)
+void UMovieGraphQuickRenderSubsystem::BeginQuickRender(const EMovieGraphQuickRenderMode InQuickRenderMode, const UMovieGraphQuickRenderModeSettings* InQuickRenderSettings)
 {
 	UMoviePipelineQueueSubsystem* Subsystem = GEditor->GetEditorSubsystem<UMoviePipelineQueueSubsystem>();
 	check(Subsystem);
@@ -90,17 +90,19 @@ void UMovieGraphQuickRenderSubsystem::BeginQuickRender(const EMovieGraphQuickRen
 		return;
 	}
 
+	QuickRenderMode = InQuickRenderMode;
+
 	// Store a reference to the settings used, so the executor callbacks can rely on them still existing.
 	QuickRenderModeSettings = InQuickRenderSettings;
 
 	// Set up the rendering and "utility" level sequences. These are both required to do a render successfully.
-	if (!SetUpAllLevelSequences(QuickRenderMode))
+	if (!SetUpAllLevelSequences())
 	{
 		return;
 	}
 
 	// Don't do additional work if the selected mode doesn't have what it needs to do a render. AreModePrerequisitesMet() will output the appropriate warnings.
-	if (!AreModePrerequisitesMet(QuickRenderMode))
+	if (!AreModePrerequisitesMet())
 	{
 		return;
 	}
@@ -151,9 +153,9 @@ void UMovieGraphQuickRenderSubsystem::BeginQuickRender(const EMovieGraphQuickRen
 	TemporaryExecutor->SetAllowUsingUnsavedLevels(true);
 
 	// Do any needed setup in the PIE world before rendering starts
-	TemporaryExecutor->OnIndividualJobStarted().AddLambda([this, QuickRenderMode, EditorWorld](UMoviePipelineExecutorJob* InJob)
+	TemporaryExecutor->OnIndividualJobStarted().AddLambda([this, EditorWorld](UMoviePipelineExecutorJob* InJob)
 	{
-		PerformPreRenderSetup(QuickRenderMode, EditorWorld);
+		PerformPreRenderSetup(EditorWorld);
 	});
 
 	// View the job's rendered frame(s)/video(s) if requested
@@ -163,9 +165,9 @@ void UMovieGraphQuickRenderSubsystem::BeginQuickRender(const EMovieGraphQuickRen
 	});
 
 	// Clean up
-	TemporaryExecutor->OnExecutorFinished().AddLambda([this, QuickRenderMode](UMoviePipelineExecutorBase* InExecutor, bool bSuccess)
+	TemporaryExecutor->OnExecutorFinished().AddLambda([this](UMoviePipelineExecutorBase* InExecutor, bool bSuccess)
 	{
-		RestorePreRenderState(QuickRenderMode);
+		RestorePreRenderState();
 		PerformPostRenderCleanup();
 	});
 
@@ -620,7 +622,7 @@ UMovieGraphNode* UMovieGraphQuickRenderSubsystem::InjectNodeIntoBranch(const TSu
 	return NewNode;
 }
 
-bool UMovieGraphQuickRenderSubsystem::AreModePrerequisitesMet(const EMovieGraphQuickRenderMode QuickRenderMode) const
+bool UMovieGraphQuickRenderSubsystem::AreModePrerequisitesMet() const
 {
 	if (QuickRenderMode == EMovieGraphQuickRenderMode::SelectedCameras)
 	{
@@ -659,7 +661,7 @@ void UMovieGraphQuickRenderSubsystem::CachePrePieData(UWorld* InEditorWorld)
 	CopyToWeakArray(CachedPrePieData.SelectedCameras, CameraActors);
 }
 
-void UMovieGraphQuickRenderSubsystem::RestorePreRenderState(const EMovieGraphQuickRenderMode QuickRenderMode)
+void UMovieGraphQuickRenderSubsystem::RestorePreRenderState()
 {
 	// Restore the modified level sequences to the way they were before the render.
 	for (const TWeakObjectPtr<ALevelSequenceActor>& LevelSequenceActor : CachedPrePieData.ModifiedLevelSequenceActors)
@@ -692,7 +694,7 @@ void UMovieGraphQuickRenderSubsystem::RestorePreRenderState(const EMovieGraphQui
 	}
 }
 
-ULevelSequence* UMovieGraphQuickRenderSubsystem::SetUpRenderingLevelSequence(const EMovieGraphQuickRenderMode QuickRenderMode) const
+ULevelSequence* UMovieGraphQuickRenderSubsystem::SetUpRenderingLevelSequence() const
 {
 	auto GetCurrentLevelSequenceOrWarn = [](EMovieGraphQuickRenderMode ActiveQuickRenderMode) -> ULevelSequence*
 	{
@@ -737,7 +739,7 @@ ULevelSequence* UMovieGraphQuickRenderSubsystem::SetUpRenderingLevelSequence(con
 	return CurrentLevelSequence;
 }
 
-ULevelSequence* UMovieGraphQuickRenderSubsystem::SetUpUtilityLevelSequence(const EMovieGraphQuickRenderMode QuickRenderMode) const
+ULevelSequence* UMovieGraphQuickRenderSubsystem::SetUpUtilityLevelSequence() const
 {
 	ULevelSequence* UtilitySequence = NewObject<ULevelSequence>(GetTransientPackage(), TEXT("QuickRenderUtilitySequence"), RF_Transient);
 	UtilitySequence->Initialize();
@@ -785,13 +787,13 @@ ULevelSequence* UMovieGraphQuickRenderSubsystem::SetUpUtilityLevelSequence(const
 	return UtilitySequence;
 }
 
-bool UMovieGraphQuickRenderSubsystem::SetUpAllLevelSequences(const EMovieGraphQuickRenderMode QuickRenderMode)
+bool UMovieGraphQuickRenderSubsystem::SetUpAllLevelSequences()
 {
 	// Get/create the level sequence that should be rendered from (varies depending on the Quick Render mode).
-	RenderingLevelSequence = SetUpRenderingLevelSequence(QuickRenderMode);
+	RenderingLevelSequence = SetUpRenderingLevelSequence();
 
 	// Also set up the utility level sequence (which is used, for example, for overriding level visibility)
-	UtilityLevelSequence = SetUpUtilityLevelSequence(QuickRenderMode);
+	UtilityLevelSequence = SetUpUtilityLevelSequence();
 
 	return IsValid(RenderingLevelSequence) && IsValid(UtilityLevelSequence);
 }
@@ -816,6 +818,12 @@ TRange<FFrameNumber> UMovieGraphQuickRenderSubsystem::GetPlaybackRange()
 	FFrameNumber StartFrame;
 	FFrameNumber EndFrame;
 
+	// If we're in Current Viewport or Selected Cameras mode, the playback range is always one frame long (only a single frame is rendered).
+	if ((QuickRenderMode == EMovieGraphQuickRenderMode::CurrentViewport) || (QuickRenderMode == EMovieGraphQuickRenderMode::SelectedCameras))
+	{
+		return TRange<FFrameNumber>::Inclusive(0, 1);
+	}
+
 	if (QuickRenderModeSettings->FrameRangeType == EMovieGraphQuickRenderFrameRangeType::SelectionRange)
 	{
 		const int32 SelectionStartFrame = ULevelSequenceEditorBlueprintLibrary::GetSelectionRangeStart();
@@ -894,7 +902,7 @@ FFrameNumber UMovieGraphQuickRenderSubsystem::ConvertSubSequenceFrameToRootFrame
 	return RootFrameTime.GetFrame();
 }
 
-void UMovieGraphQuickRenderSubsystem::PerformPreRenderSetup(const EMovieGraphQuickRenderMode QuickRenderMode, UWorld* InEditorWorld)
+void UMovieGraphQuickRenderSubsystem::PerformPreRenderSetup(UWorld* InEditorWorld)
 {
 	// Adds the specified camera to the Camera Cut track in the given level sequence.
 	// If there's already a camera cut section within the Camera Cut track, a new one will be added immediately after.
diff --git a/Engine/Plugins/MovieScene/MovieRenderPipeline/Source/MovieRenderPipelineEditor/Private/MoviePipelineQuickRenderMenu.cpp b/Engine/Plugins/MovieScene/MovieRenderPipeline/Source/MovieRenderPipelineEditor/Private/MoviePipelineQuickRenderMenu.cpp
--- a/Engine/Plugins/MovieScene/MovieRenderPipeline/Source/MovieRenderPipelineEditor/Private/MoviePipelineQuickRenderMenu.cpp
+++ b/Engine/Plugins/MovieScene/MovieRenderPipeline/Source/MovieRenderPipelineEditor/Private/MoviePipelineQuickRenderMenu.cpp
@@ -9,6 +9,7 @@
 #include "Graph/MovieGraphQuickRender.h"
 #include "Misc/ConfigCacheIni.h"
 #include "Misc/MessageDialog.h"
+#include "MoviePipelineQuickRenderUIState.h"
 #include "MoviePipelineUtils.h"
 #include "MovieRenderPipelineCoreModule.h"
 #include "MovieRenderPipelineEditorModule.h"
@@ -83,24 +84,24 @@ void FMoviePipelineQuickRenderMenu::RemoveQuickRenderButtonToolMenu(UToolMenu* I
 
 void FMoviePipelineQuickRenderMenu::LoadQuickRenderSettings()
 {
-	InitQuickRenderModeSettingsFromMode(GetQuickRenderMode());
+	InitQuickRenderModeSettingsFromMode(FMoviePipelineQuickRenderUIState::GetQuickRenderMode());
 }
 
 void FMoviePipelineQuickRenderMenu::QuickRenderButtonPressed()
 {
 	LoadQuickRenderSettings();
 
-	if (GetQuickRenderButtonMode() == EMovieGraphQuickRenderButtonMode::QuickRender)
+	if (FMoviePipelineQuickRenderUIState::GetQuickRenderButtonMode() == EMovieGraphQuickRenderButtonMode::QuickRender)
 	{
 		// The settings window will have its own Render button, so don't perform the render if showing the settings first
-		if (GetShouldShowSettingsBeforeRender())
+		if (FMoviePipelineQuickRenderUIState::GetShouldShowSettingsBeforeRender())
 		{
 			OpenQuickRenderSettingsWindow(FToolMenuContext());
 		}
 		else
 		{
 			UMovieGraphQuickRenderSubsystem* QuickRenderSubsystem = GEditor->GetEditorSubsystem<UMovieGraphQuickRenderSubsystem>();
-			QuickRenderSubsystem->BeginQuickRender(GetQuickRenderMode(), QuickRenderModeSettings.Get());
+			QuickRenderSubsystem->BeginQuickRender(FMoviePipelineQuickRenderUIState::GetQuickRenderMode(), QuickRenderModeSettings.Get());
 		}
 	}
 	else
@@ -153,7 +154,7 @@ void FMoviePipelineQuickRenderMenu::OpenQuickRenderSettingsWindow(const FToolMen
 	}
 
 	// Before the window opens, sync up the mode that the Settings window is using to be the mode that the toolbar shows
-	SettingsWindowQuickRenderMode = GetQuickRenderMode();
+	FMoviePipelineQuickRenderUIState::SetWindowQuickRenderMode(FMoviePipelineQuickRenderUIState::GetQuickRenderMode());
 
 	// Init the modes that are available to switch to
 	QuickRenderModes.Empty();
@@ -260,13 +261,14 @@ TSharedRef<SWidget> FMoviePipelineQuickRenderMenu::GenerateQuickRenderSettingsWi
 		})
 		.OnSelectionChanged_Lambda([](TSharedPtr<FName> InSelectedItem, ESelectInfo::Type)
 		{
-			SettingsWindowQuickRenderMode = static_cast<EMovieGraphQuickRenderMode>(
+			const EMovieGraphQuickRenderMode NewRenderMode = static_cast<EMovieGraphQuickRenderMode>(
 				StaticEnum<EMovieGraphQuickRenderMode>()->GetValueByName(*InSelectedItem));
+			FMoviePipelineQuickRenderUIState::SetWindowQuickRenderMode(NewRenderMode);
 
 			// Also inform the details panel of this change. It needs to display new settings for the chosen mode.
 			if (const TSharedPtr<IDetailsView> DetailsPanel = WeakDetailsPanel.Pin())
 			{
-				InitQuickRenderModeSettingsFromMode(SettingsWindowQuickRenderMode);
+				InitQuickRenderModeSettingsFromMode(NewRenderMode);
 				
 				DetailsPanel->SetObject(QuickRenderModeSettings.Get());
 			}
@@ -274,7 +276,7 @@ TSharedRef<SWidget> FMoviePipelineQuickRenderMenu::GenerateQuickRenderSettingsWi
 		[
 			GetWidgetForQuickRenderMode([]()
 			{
-				return SettingsWindowQuickRenderMode;
+				return FMoviePipelineQuickRenderUIState::GetWindowQuickRenderMode();
 			})
 		], NAME_None, false, HAlign_Left);
 
@@ -286,10 +288,10 @@ TSharedRef<SWidget> FMoviePipelineQuickRenderMenu::GenerateQuickRenderSettingsWi
 		.OnClicked_Lambda([]()
 		{
 			// Sync the Setting window's mode to the toolbar mode. This is only done when a render is performed, not when the window is closed.
-			SetQuickRenderMode(SettingsWindowQuickRenderMode);
+			FMoviePipelineQuickRenderUIState::SetQuickRenderMode(FMoviePipelineQuickRenderUIState::GetWindowQuickRenderMode());
 			
 			UMovieGraphQuickRenderSubsystem* QuickRenderSubsystem = GEditor->GetEditorSubsystem<UMovieGraphQuickRenderSubsystem>();
-			QuickRenderSubsystem->BeginQuickRender(GetQuickRenderMode(), QuickRenderModeSettings.Get());
+			QuickRenderSubsystem->BeginQuickRender(FMoviePipelineQuickRenderUIState::GetQuickRenderMode(), QuickRenderModeSettings.Get());
 			
 			return FReply::Handled();
 		}),
@@ -309,13 +311,13 @@ void FMoviePipelineQuickRenderMenu::GenerateModesMenuSection(UToolMenu* InMenu)
 		FToolUIAction NewAction;
 		NewAction.ExecuteAction = FToolMenuExecuteAction::CreateLambda([CheckedValue](const FToolMenuContext&)
 		{
-			SetQuickRenderButtonMode(CheckedValue);
+			FMoviePipelineQuickRenderUIState::SetQuickRenderButtonMode(CheckedValue);
 			QuickRenderModeSettings->PostEditChange();
 		});
 		NewAction.CanExecuteAction = FToolMenuCanExecuteAction::CreateLambda([](const FToolMenuContext&) { return true; });
 		NewAction.GetActionCheckState = FToolMenuGetActionCheckState::CreateLambda([CheckedValue](const FToolMenuContext&)
 		{
-			return (GetQuickRenderButtonMode() == CheckedValue)
+			return (FMoviePipelineQuickRenderUIState::GetQuickRenderButtonMode() == CheckedValue)
 				? ECheckBoxState::Checked
 				: ECheckBoxState::Unchecked;
 		});
@@ -346,13 +348,13 @@ void FMoviePipelineQuickRenderMenu::GenerateQuickRenderMenuSection(UToolMenu* In
 		FToolUIAction NewAction;
 		NewAction.ExecuteAction = FToolMenuExecuteAction::CreateLambda([CheckedValue](const FToolMenuContext&)
 		{
-			SetQuickRenderMode(CheckedValue);
+			FMoviePipelineQuickRenderUIState::SetQuickRenderMode(CheckedValue);
 			QuickRenderModeSettings->PostEditChange();
 		});
 		NewAction.CanExecuteAction = FToolMenuCanExecuteAction::CreateLambda([](const FToolMenuContext&) { return true; });
 		NewAction.GetActionCheckState = FToolMenuGetActionCheckState::CreateLambda([CheckedValue](const FToolMenuContext&)
 		{
-			return (GetQuickRenderMode() == CheckedValue)
+			return (FMoviePipelineQuickRenderUIState::GetQuickRenderMode() == CheckedValue)
 				? ECheckBoxState::Checked
 				: ECheckBoxState::Unchecked;
 		});
@@ -462,7 +464,9 @@ void FMoviePipelineQuickRenderMenu::GenerateOutputMenuSection(UToolMenu* InMenu)
 	FToolMenuSection& Section = InMenu->AddSection("MoviePipelineQuickRenderOutput", LOCTEXT("QuickRenderOutput", "Output"));
 	Section.Visibility = TAttribute<EVisibility>::CreateLambda([]() -> EVisibility
 	{
-		return (GetQuickRenderButtonMode() == EMovieGraphQuickRenderButtonMode::QuickRender) ? EVisibility::Visible : EVisibility::Collapsed;
+		return (FMoviePipelineQuickRenderUIState::GetQuickRenderButtonMode() == EMovieGraphQuickRenderButtonMode::QuickRender)
+			? EVisibility::Visible
+			: EVisibility::Collapsed;
 	});
 
 	// Menu entry for "Play Last Render"
@@ -511,13 +515,13 @@ void FMoviePipelineQuickRenderMenu::GenerateSettingsMenuSection(UToolMenu* InMen
 	FToolUIAction ShowSettingsAction;
         ShowSettingsAction.ExecuteAction = FToolMenuExecuteAction::CreateLambda([](const FToolMenuContext&)
         {
-			SetShouldShowSettingsBeforeRender(!GetShouldShowSettingsBeforeRender());
+			FMoviePipelineQuickRenderUIState::SetShouldShowSettingsBeforeRender(!FMoviePipelineQuickRenderUIState::GetShouldShowSettingsBeforeRender());
         	QuickRenderModeSettings->PostEditChange();
         });
         ShowSettingsAction.CanExecuteAction = FToolMenuCanExecuteAction::CreateLambda([](const FToolMenuContext&) { return true; });
         ShowSettingsAction.GetActionCheckState = FToolMenuGetActionCheckState::CreateLambda([](const FToolMenuContext&)
         {
-            return GetShouldShowSettingsBeforeRender()
+            return FMoviePipelineQuickRenderUIState::GetShouldShowSettingsBeforeRender()
 				? ECheckBoxState::Checked
 				: ECheckBoxState::Unchecked;
         });
@@ -568,42 +572,4 @@ void FMoviePipelineQuickRenderMenu::UpdateVariableAssignmentsForCurrentGraph()
 	MoviePipeline::RefreshVariableAssignments(QuickRenderModeSettings->GraphPreset.LoadSynchronous(), QuickRenderModeSettings->GraphVariableAssignments, QuickRenderModeSettings.Get());
 }
 
-EMovieGraphQuickRenderButtonMode FMoviePipelineQuickRenderMenu::GetQuickRenderButtonMode()
-{
-	constexpr EMovieGraphQuickRenderButtonMode DefaultButtonMode = EMovieGraphQuickRenderButtonMode::QuickRender;
-	
-	return static_cast<EMovieGraphQuickRenderButtonMode>(
-		GConfig->GetIntOrDefault(QuickRenderIniSection, QuickRenderIniSettingName_ButtonMode, static_cast<int32>(DefaultButtonMode), GEditorPerProjectIni));
-}
-
-void FMoviePipelineQuickRenderMenu::SetQuickRenderButtonMode(const EMovieGraphQuickRenderButtonMode NewButtonMode)
-{
-	GConfig->SetInt(QuickRenderIniSection, QuickRenderIniSettingName_ButtonMode, static_cast<int32>(NewButtonMode), GEditorPerProjectIni);
-}
-
-EMovieGraphQuickRenderMode FMoviePipelineQuickRenderMenu::GetQuickRenderMode()
-{
-	constexpr EMovieGraphQuickRenderMode DefaultMode = EMovieGraphQuickRenderMode::CurrentSequence;
-	
-	return static_cast<EMovieGraphQuickRenderMode>(
-		GConfig->GetIntOrDefault(QuickRenderIniSection, QuickRenderIniSettingName_Mode, static_cast<int32>(DefaultMode), GEditorPerProjectIni));
-}
-
-void FMoviePipelineQuickRenderMenu::SetQuickRenderMode(const EMovieGraphQuickRenderMode NewMode)
-{
-	GConfig->SetInt(QuickRenderIniSection, QuickRenderIniSettingName_Mode, static_cast<int32>(NewMode), GEditorPerProjectIni);
-}
-
-bool FMoviePipelineQuickRenderMenu::GetShouldShowSettingsBeforeRender()
-{
-	constexpr bool bDefaultShowSettingsBeforeRender = false;
-	
-	return GConfig->GetBoolOrDefault(QuickRenderIniSection, QuickRenderIniSettingName_ShowSettingsBeforeQuickRender, bDefaultShowSettingsBeforeRender, GEditorPerProjectIni);
-}
-
-void FMoviePipelineQuickRenderMenu::SetShouldShowSettingsBeforeRender(const bool bNewValue)
-{
-	GConfig->SetInt(QuickRenderIniSection, QuickRenderIniSettingName_ShowSettingsBeforeQuickRender, bNewValue, GEditorPerProjectIni);
-}
-
 #undef LOCTEXT_NAMESPACE
\ No newline at end of file
diff --git a/Engine/Plugins/MovieScene/MovieRenderPipeline/Source/MovieRenderPipelineEditor/Private/MoviePipelineQuickRenderMenu.h b/Engine/Plugins/MovieScene/MovieRenderPipeline/Source/MovieRenderPipelineEditor/Private/MoviePipelineQuickRenderMenu.h
--- a/Engine/Plugins/MovieScene/MovieRenderPipeline/Source/MovieRenderPipelineEditor/Private/MoviePipelineQuickRenderMenu.h
+++ b/Engine/Plugins/MovieScene/MovieRenderPipeline/Source/MovieRenderPipelineEditor/Private/MoviePipelineQuickRenderMenu.h
@@ -46,41 +46,7 @@ private:
 	/** Updates the variable assignments for the graph preset in the currently active mode. */
 	static void UpdateVariableAssignmentsForCurrentGraph();
 
-	// ---- UI state getters and setters ----------------------------------------------------------------------
-
-	/** Gets the mode that the Quick Render button is currently in. */
-	static EMovieGraphQuickRenderButtonMode GetQuickRenderButtonMode();
-
-	/** Sets the mode that the Quick Render button is currently using. */
-	static void SetQuickRenderButtonMode(const EMovieGraphQuickRenderButtonMode NewButtonMode);
-
-	/** Gets the mode that Quick Render will use when starting a new render. */
-	static EMovieGraphQuickRenderMode GetQuickRenderMode();
-
-	/** Sets the mode that Quick Render will use when starting a new render. */
-	static void SetQuickRenderMode(const EMovieGraphQuickRenderMode NewMode);
-
-	/** Gets whether Quick Render should show the Settings dialog before a render starts. */
-	static bool GetShouldShowSettingsBeforeRender();
-
-	/** Sets whether Quick Render should show the Settings dialog before a render starts. */
-	static void SetShouldShowSettingsBeforeRender(const bool bNewValue);
-
-	// --------------------------------------------------------------------------------------------------------
-
 private:
-	/** The name of the section that Quick Render settings are stored in within the target ini file. */
-	inline static const TCHAR* QuickRenderIniSection = TEXT("MovieRenderPipeline.QuickRender");
-
-	/** Ini setting name for: What mode the Quick Render button is using. */
-	inline static const TCHAR* QuickRenderIniSettingName_ButtonMode = TEXT("ButtonMode");
-
-	/** Ini setting name for: The mode that Quick Render will use when a render begins. Generally dictates the level sequence and camera(s) that will be used. */
-	inline static const TCHAR* QuickRenderIniSettingName_Mode = TEXT("Mode");
-
-	/** Ini setting name for: Whether the settings dialog should be shown before a Quick Render is started. */
-	inline static const TCHAR* QuickRenderIniSettingName_ShowSettingsBeforeQuickRender = TEXT("bShowSettingsBeforeQuickRender");
-
 	/** The ID of the Quick Render toolbar that's added to the viewport toolbar. */
 	inline static const FName QuickRenderToolbarId = TEXT("MovieRenderGraphQuickRender");
 
@@ -95,7 +61,4 @@ private:
 
 	/** All of the different mode names that are available to be used. */
 	inline static TArray<TSharedPtr<FName>> QuickRenderModes;
-
-	/** The Quick Render mode that the Settings window is currently using (vs. the mode that's active in the toolbar). */
-	inline static EMovieGraphQuickRenderMode SettingsWindowQuickRenderMode;
 };
diff --git a/Engine/Plugins/MovieScene/MovieRenderPipeline/Source/MovieRenderPipelineEditor/Private/MoviePipelineQuickRenderUIState.cpp b/Engine/Plugins/MovieScene/MovieRenderPipeline/Source/MovieRenderPipelineEditor/Private/MoviePipelineQuickRenderUIState.cpp
new file mode 100644
--- /dev/null
+++ b/Engine/Plugins/MovieScene/MovieRenderPipeline/Source/MovieRenderPipelineEditor/Private/MoviePipelineQuickRenderUIState.cpp
@@ -0,0 +1,53 @@
+﻿// Copyright Epic Games, Inc. All Rights Reserved.
+
+#include "MoviePipelineQuickRenderUIState.h"
+
+#include "Misc/ConfigCacheIni.h"
+
+EMovieGraphQuickRenderButtonMode FMoviePipelineQuickRenderUIState::GetQuickRenderButtonMode()
+{
+	constexpr EMovieGraphQuickRenderButtonMode DefaultButtonMode = EMovieGraphQuickRenderButtonMode::QuickRender;
+	
+	return static_cast<EMovieGraphQuickRenderButtonMode>(
+		GConfig->GetIntOrDefault(QuickRenderIniSection, QuickRenderIniSettingName_ButtonMode, static_cast<int32>(DefaultButtonMode), GEditorPerProjectIni));
+}
+
+void FMoviePipelineQuickRenderUIState::SetQuickRenderButtonMode(const EMovieGraphQuickRenderButtonMode NewButtonMode)
+{
+	GConfig->SetInt(QuickRenderIniSection, QuickRenderIniSettingName_ButtonMode, static_cast<int32>(NewButtonMode), GEditorPerProjectIni);
+}
+
+EMovieGraphQuickRenderMode FMoviePipelineQuickRenderUIState::GetQuickRenderMode()
+{
+	constexpr EMovieGraphQuickRenderMode DefaultMode = EMovieGraphQuickRenderMode::CurrentSequence;
+	
+	return static_cast<EMovieGraphQuickRenderMode>(
+		GConfig->GetIntOrDefault(QuickRenderIniSection, QuickRenderIniSettingName_Mode, static_cast<int32>(DefaultMode), GEditorPerProjectIni));
+}
+
+EMovieGraphQuickRenderMode FMoviePipelineQuickRenderUIState::GetWindowQuickRenderMode()
+{
+	return WindowRenderMode;
+}
+
+void FMoviePipelineQuickRenderUIState::SetWindowQuickRenderMode(const EMovieGraphQuickRenderMode NewMode)
+{
+	WindowRenderMode = NewMode;
+}
+
+void FMoviePipelineQuickRenderUIState::SetQuickRenderMode(const EMovieGraphQuickRenderMode NewMode)
+{
+	GConfig->SetInt(QuickRenderIniSection, QuickRenderIniSettingName_Mode, static_cast<int32>(NewMode), GEditorPerProjectIni);
+}
+
+bool FMoviePipelineQuickRenderUIState::GetShouldShowSettingsBeforeRender()
+{
+	constexpr bool bDefaultShowSettingsBeforeRender = false;
+	
+	return GConfig->GetBoolOrDefault(QuickRenderIniSection, QuickRenderIniSettingName_ShowSettingsBeforeQuickRender, bDefaultShowSettingsBeforeRender, GEditorPerProjectIni);
+}
+
+void FMoviePipelineQuickRenderUIState::SetShouldShowSettingsBeforeRender(const bool bNewValue)
+{
+	GConfig->SetInt(QuickRenderIniSection, QuickRenderIniSettingName_ShowSettingsBeforeQuickRender, bNewValue, GEditorPerProjectIni);
+}
\ No newline at end of file
diff --git a/Engine/Plugins/MovieScene/MovieRenderPipeline/Source/MovieRenderPipelineEditor/Private/MoviePipelineQuickRenderUIState.h b/Engine/Plugins/MovieScene/MovieRenderPipeline/Source/MovieRenderPipelineEditor/Private/MoviePipelineQuickRenderUIState.h
new file mode 100644
--- /dev/null
+++ b/Engine/Plugins/MovieScene/MovieRenderPipeline/Source/MovieRenderPipelineEditor/Private/MoviePipelineQuickRenderUIState.h
@@ -0,0 +1,62 @@
+﻿// Copyright Epic Games, Inc. All Rights Reserved.
+
+#pragma once
+
+#include "Graph/MovieGraphQuickRenderSettings.h"
+
+/**
+ * UI state management for Quick Render.
+ *
+ * Some of these setting values returned here may differ from what Quick Render is actively using. For example, the mode returned here may be
+ * different if Quick Render was started from scripting with a mode that differs from the UI's current mode. Also note that the settings here are
+ * distinct from the settings in UMovieGraphQuickRenderSettings. Those are persisted to an asset, whereas these are persisted to an INI. 
+ */
+struct FMoviePipelineQuickRenderUIState
+{
+	/** Gets the mode that the Quick Render button is currently in. */
+	static EMovieGraphQuickRenderButtonMode GetQuickRenderButtonMode();
+
+	/** Sets the mode that the Quick Render button is currently using. */
+	static void SetQuickRenderButtonMode(const EMovieGraphQuickRenderButtonMode NewButtonMode);
+
+	/**
+	 * Gets the mode that Quick Render will use when starting a new render. Sometimes called the "system" render mode, as opposed to the "window"
+	 * render mode (see GetWindowQuickRenderMode()).
+	 */
+	static EMovieGraphQuickRenderMode GetQuickRenderMode();
+
+	/** Sets the mode that Quick Render will use when starting a new render. */
+	static void SetQuickRenderMode(const EMovieGraphQuickRenderMode NewMode);
+
+	/**
+	 * Gets the render mode that the Quick Render settings window currently displays. This is different than GetQuickRenderMode(), or the "system"
+	 * render mode. When the settings window opens, this mode is independent of the system render mode. When a render starts, the window mode is
+	 * copied to the system render mode. This setting is transient and not persisted across editor sessions.
+	 */
+	static EMovieGraphQuickRenderMode GetWindowQuickRenderMode();
+
+	/** Sets the "window" render mode. See notes on GetWindowQuickRenderMode(). */
+	static void SetWindowQuickRenderMode(const EMovieGraphQuickRenderMode NewMode);
+
+	/** Gets whether Quick Render should show the Settings dialog before a render starts. */
+	static bool GetShouldShowSettingsBeforeRender();
+
+	/** Sets whether Quick Render should show the Settings dialog before a render starts. */
+	static void SetShouldShowSettingsBeforeRender(const bool bNewValue);
+
+private:
+	/** The name of the section that Quick Render UI settings are stored in within the target ini file. */
+	inline static const TCHAR* QuickRenderIniSection = TEXT("MovieRenderPipeline.QuickRender");
+
+	/** Ini setting name for: What mode the Quick Render button is using. */
+	inline static const TCHAR* QuickRenderIniSettingName_ButtonMode = TEXT("ButtonMode");
+
+	/** Ini setting name for: The mode that Quick Render will use when a render begins. Generally dictates the level sequence and camera(s) that will be used. */
+	inline static const TCHAR* QuickRenderIniSettingName_Mode = TEXT("Mode");
+
+	/** Ini setting name for: Whether the settings dialog should be shown before a Quick Render is started. */
+	inline static const TCHAR* QuickRenderIniSettingName_ShowSettingsBeforeQuickRender = TEXT("bShowSettingsBeforeQuickRender");
+
+	/** The transient "window" quick render mode. This is not persisted to ini. */
+	static inline EMovieGraphQuickRenderMode WindowRenderMode = EMovieGraphQuickRenderMode::CurrentViewport;
+};
diff --git a/Engine/Plugins/MovieScene/MovieRenderPipeline/Source/MovieRenderPipelineEditor/Public/Graph/MovieGraphQuickRender.h b/Engine/Plugins/MovieScene/MovieRenderPipeline/Source/MovieRenderPipelineEditor/Public/Graph/MovieGraphQuickRender.h
--- a/Engine/Plugins/MovieScene/MovieRenderPipeline/Source/MovieRenderPipelineEditor/Public/Graph/MovieGraphQuickRender.h
+++ b/Engine/Plugins/MovieScene/MovieRenderPipeline/Source/MovieRenderPipelineEditor/Public/Graph/MovieGraphQuickRender.h
@@ -27,7 +27,7 @@ class MOVIERENDERPIPELINEEDITOR_API UMovieGraphQuickRenderSubsystem : public UEd
 public:
 	/* Begins a quick render using the provided mode and settings. */
 	UFUNCTION(BlueprintCallable, Category="Quick Render")
-	void BeginQuickRender(const EMovieGraphQuickRenderMode QuickRenderMode, const UMovieGraphQuickRenderModeSettings* InQuickRenderSettings);
+	void BeginQuickRender(const EMovieGraphQuickRenderMode InQuickRenderMode, const UMovieGraphQuickRenderModeSettings* InQuickRenderSettings);
 
 	/**
 	 * Plays the last render that Quick Render generated, using the settings specified in Editor Preferences. If no render has been
@@ -89,31 +89,31 @@ private:
 	UMovieGraphNode* InjectNodeIntoBranch(const TSubclassOf<UMovieGraphNode> NodeType, const FName& InBranchName) const;
 
 	/** Determines if the prerequisites for the given mode are met. */
-	bool AreModePrerequisitesMet(const EMovieGraphQuickRenderMode QuickRenderMode) const;
+	bool AreModePrerequisitesMet() const;
 
 	/** Caches any necessary data before PIE starts. Stored in the CachedPrePieData member variable. */
 	void CachePrePieData(UWorld* InEditorWorld);
 
 	/** Handles restoring data after PIE ends. */
-	void RestorePreRenderState(const EMovieGraphQuickRenderMode QuickRenderMode);
+	void RestorePreRenderState();
 
 	/**
 	 * Returns the current level sequence active in Sequencer, or sets up a new level sequence for rendering (depending on the mode that
 	 * Quick Render is using).
 	 */
-	ULevelSequence* SetUpRenderingLevelSequence(const EMovieGraphQuickRenderMode QuickRenderMode) const;
+	ULevelSequence* SetUpRenderingLevelSequence() const;
 
 	/**
 	 * Returns a new level sequence used for utility purposes (like overriding the camera in use, or changing level
 	 * visibility). It will be injected into the rendering level sequence as a subsequence.
 	 */
-	ULevelSequence* SetUpUtilityLevelSequence(const EMovieGraphQuickRenderMode QuickRenderMode) const;
+	ULevelSequence* SetUpUtilityLevelSequence() const;
 
 	/**
 	 * Populates both the rendering and utility level sequences, and assigns them to the relevant members. Returns true if everything could be set
 	 * up, otherwise false.
 	 */
-	bool SetUpAllLevelSequences(const EMovieGraphQuickRenderMode QuickRenderMode);
+	bool SetUpAllLevelSequences();
 
 	/** Gets the first world found of the specified type. */
 	UWorld* GetWorldOfType(const EWorldType::Type WorldType) const;
@@ -131,7 +131,7 @@ private:
 	FFrameNumber ConvertSubSequenceFrameToRootFrame(const int32 FrameNum) const;
 
 	/** Does any setup needed in the PIE world prior to the render starting. */
-	void PerformPreRenderSetup(const EMovieGraphQuickRenderMode QuickRenderMode, UWorld* InEditorWorld);
+	void PerformPreRenderSetup(UWorld* InEditorWorld);
 
 	/** Opens the rendered files that are in the given output data. */
 	void OpenPostRenderFileDisplayProcessor(const FMoviePipelineOutputData& InOutputData) const;
@@ -173,6 +173,14 @@ private:
     UPROPERTY(Transient)
     TObjectPtr<ULevelSequence> UtilityLevelSequence = nullptr;
 
+	/**
+	 * The mode that Quick Render is actively using in a render. This may be different from the mode that's set in the INI file, which indicates the
+	 * mode that the UI is using (which could differ if Quick Render is triggered with scripting, for example). See QuickRenderModeSettings for mode
+	 * configuration settings.
+	 */
+	UPROPERTY(Transient)
+	EMovieGraphQuickRenderMode QuickRenderMode = EMovieGraphQuickRenderMode::CurrentViewport;
+
 	/**
 	 * The mode settings Quick Render was initialized with. These are cached here to grab a reference to them before the render begins so the
 	 * executor callbacks still have access to them later.
