MoviePipeline: Moved the Quick Render button to a new place, to the right of the Level Sequences button in the main toolbar.

diff --git a/Engine/Plugins/MovieScene/MovieRenderPipeline/Source/MovieRenderPipelineEditor/Private/MoviePipelineQuickRenderMenu.cpp b/Engine/Plugins/MovieScene/MovieRenderPipeline/Source/MovieRenderPipelineEditor/Private/MoviePipelineQuickRenderMenu.cpp
--- a/Engine/Plugins/MovieScene/MovieRenderPipeline/Source/MovieRenderPipelineEditor/Private/MoviePipelineQuickRenderMenu.cpp
+++ b/Engine/Plugins/MovieScene/MovieRenderPipeline/Source/MovieRenderPipelineEditor/Private/MoviePipelineQuickRenderMenu.cpp
@@ -36,15 +36,15 @@ void FMoviePipelineQuickRenderMenu::AddQuickRenderButtonToToolMenu(UToolMenu* In
 		return;
 	}
 	
-	FToolMenuSection& Section = InMenu->AddSection(QuickRenderToolbarId);
-	Section.AddSeparator("QuickRenderSeparator");
+	// Add to the existing "Content" section in the Assets toolbar
+	FToolMenuSection& Section = InMenu->FindOrAddSection("Content");
 	
 	const FToolMenuEntry QuickRenderLaunchButton =
 		FToolMenuEntry::InitToolBarButton(
-			"QuickRenderToolbarButton",
+			"QuickRender",
 			FUIAction(
 				FExecuteAction::CreateStatic(&FMoviePipelineQuickRenderMenu::QuickRenderButtonPressed)),
-			TAttribute<FText>(),
+			LOCTEXT("QuickRenderButtonName", "Quick Render"),
 			TAttribute<FText>::Create(
 				TAttribute<FText>::FGetter::CreateLambda([]()
 				{
@@ -69,16 +69,11 @@ void FMoviePipelineQuickRenderMenu::AddQuickRenderButtonToToolMenu(UToolMenu* In
 	Section.AddEntry(QuickRenderOptionsButton);
 }
 
-void FMoviePipelineQuickRenderMenu::RemoveQuickRenderButtonToolMenu(UToolMenu* InMenu)
+void FMoviePipelineQuickRenderMenu::RemoveQuickRenderButtonToolMenu()
 {
-	if (!InMenu)
-	{
-		return;
-	}
-	
-	if (InMenu->FindSection(QuickRenderToolbarId))
+	if (UToolMenus* ToolMenus = UToolMenus::TryGet())
 	{
-		InMenu->RemoveSection(QuickRenderToolbarId);
+		ToolMenus->RemoveEntry("LevelEditor.LevelEditorToolBar.AssetsToolBar", "Content", "QuickRender");
 	}
 }
 
diff --git a/Engine/Plugins/MovieScene/MovieRenderPipeline/Source/MovieRenderPipelineEditor/Private/MoviePipelineQuickRenderMenu.h b/Engine/Plugins/MovieScene/MovieRenderPipeline/Source/MovieRenderPipelineEditor/Private/MoviePipelineQuickRenderMenu.h
--- a/Engine/Plugins/MovieScene/MovieRenderPipeline/Source/MovieRenderPipelineEditor/Private/MoviePipelineQuickRenderMenu.h
+++ b/Engine/Plugins/MovieScene/MovieRenderPipeline/Source/MovieRenderPipelineEditor/Private/MoviePipelineQuickRenderMenu.h
@@ -20,8 +20,8 @@ public:
 	/** Adds the Quick Render button to the specified tool menu. */
 	static void AddQuickRenderButtonToToolMenu(UToolMenu* InMenu);
 
-	/** Removes the Quick Render button from the specified tool menu. */
-	static void RemoveQuickRenderButtonToolMenu(UToolMenu* InMenu);
+	/** Removes the Quick Render button from level editor toolbar. */
+	static void RemoveQuickRenderButtonToolMenu();
 
 private:
 	static void LoadQuickRenderSettings();
@@ -47,9 +47,6 @@ private:
 	static void UpdateVariableAssignmentsForCurrentGraph();
 
 private:
-	/** The ID of the Quick Render toolbar that's added to the viewport toolbar. */
-	inline static const FName QuickRenderToolbarId = TEXT("MovieRenderGraphQuickRender");
-
 	/** Weak pointer to the Settings window widget. */
 	inline static TWeakPtr<SWindow> WeakQuickRenderSettingsWindow;
 
diff --git a/Engine/Plugins/MovieScene/MovieRenderPipeline/Source/MovieRenderPipelineEditor/Private/MovieRenderPipelineEditorModule.cpp b/Engine/Plugins/MovieScene/MovieRenderPipeline/Source/MovieRenderPipelineEditor/Private/MovieRenderPipelineEditorModule.cpp
--- a/Engine/Plugins/MovieScene/MovieRenderPipeline/Source/MovieRenderPipelineEditor/Private/MovieRenderPipelineEditorModule.cpp
+++ b/Engine/Plugins/MovieScene/MovieRenderPipeline/Source/MovieRenderPipelineEditor/Private/MovieRenderPipelineEditorModule.cpp
@@ -237,7 +237,7 @@ void FMovieRenderPipelineEditorModule::UnregisterTypeCustomizations()
 
 void FMovieRenderPipelineEditorModule::RegisterToolbarItem()
 {
-	UToolMenu* Menu = UToolMenus::Get()->ExtendMenu("LevelEditor.LevelEditorToolBar.User");
+	UToolMenu* Menu = UToolMenus::Get()->ExtendMenu("LevelEditor.LevelEditorToolBar.AssetsToolBar");
 	
 	UnregisterToolbarItem();
 	
@@ -246,11 +246,7 @@ void FMovieRenderPipelineEditorModule::RegisterToolbarItem()
 
 void FMovieRenderPipelineEditorModule::UnregisterToolbarItem()
 {
-	if (UToolMenus* ToolMenus = UToolMenus::Get())
-	{
-		UToolMenu* Menu = ToolMenus->FindMenu("LevelEditor.LevelEditorToolBar.User");
-		FMoviePipelineQuickRenderMenu::RemoveQuickRenderButtonToolMenu(Menu);
-	}
+	FMoviePipelineQuickRenderMenu::RemoveQuickRenderButtonToolMenu();
 }
 
 void FMovieRenderPipelineEditorModule::StartupModule()
