MoviePipeline: Alpha output renderer setting (r.PostProcessing.PropagateAlpha) is now automatically enabled by default during renders.

It can be disabled by setting MoviePipeline.AlphaOutputOverride to false/0.

diff --git a/Engine/Plugins/MovieScene/MoviePipelineMaskRenderPass/Source/MoviePipelineMaskRenderPass/Private/MoviePipelinePanoramicPass.h b/Engine/Plugins/MovieScene/MoviePipelineMaskRenderPass/Source/MoviePipelineMaskRenderPass/Private/MoviePipelinePanoramicPass.h
--- a/Engine/Plugins/MovieScene/MoviePipelineMaskRenderPass/Source/MoviePipelineMaskRenderPass/Private/MoviePipelinePanoramicPass.h
+++ b/Engine/Plugins/MovieScene/MoviePipelineMaskRenderPass/Source/MoviePipelineMaskRenderPass/Private/MoviePipelinePanoramicPass.h
@@ -125,7 +125,6 @@ protected:
 	virtual void GatherOutputPassesImpl(TArray<FMoviePipelinePassIdentifier>& ExpectedRenderPasses) override;
 	virtual bool IsAntiAliasingSupported() const { return true; }
 	virtual int32 GetOutputFileSortingOrder() const override { return 1; }
-	virtual bool IsAlphaInTonemapperRequiredImpl() const override { return false; }
 	virtual FSceneViewStateInterface* GetSceneViewStateInterface(IViewCalcPayload* OptPayload) override;
 	virtual FSceneViewStateInterface* GetExposureSceneViewStateInterface(IViewCalcPayload* OptPayload, int32 CubeFaceIndex) override;
 	virtual void AddViewExtensions(FSceneViewFamilyContext& InContext, FMoviePipelineRenderPassMetrics& InOutSampleState) override;
diff --git a/Engine/Plugins/MovieScene/MovieRenderPipeline/Source/MovieRenderPipelineCore/Private/Graph/Nodes/MovieGraphGlobalGameOverrides.cpp b/Engine/Plugins/MovieScene/MovieRenderPipeline/Source/MovieRenderPipelineCore/Private/Graph/Nodes/MovieGraphGlobalGameOverrides.cpp
--- a/Engine/Plugins/MovieScene/MovieRenderPipeline/Source/MovieRenderPipelineCore/Private/Graph/Nodes/MovieGraphGlobalGameOverrides.cpp
+++ b/Engine/Plugins/MovieScene/MovieRenderPipeline/Source/MovieRenderPipelineCore/Private/Graph/Nodes/MovieGraphGlobalGameOverrides.cpp
@@ -95,6 +95,8 @@ void UMovieGraphGlobalGameOverridesNode::BuildNewProcessCommandLineArgsImpl(TArr
 	{
 		InOutDeviceProfileCvars.AddUnique(FString::Printf(TEXT("r.Nanite.VSMInvalidateOnLODDelta=%d"), 1));
 	}
+
+	InOutDeviceProfileCvars.AddUnique(FString::Printf(TEXT("r.PostProcessing.PropagateAlpha=%d"), UE::MoviePipeline::GetAlphaOutputOverride() ? 1 : 0));
 }
 
 void UMovieGraphGlobalGameOverridesNode::PostLoad()
@@ -171,6 +173,8 @@ void UMovieGraphGlobalGameOverridesNode::ApplySettings(const bool bOverrideValue
 		MOVIEPIPELINE_STORE_AND_OVERRIDE_CVAR_INT(NodeCDO->PreviousStreamingManagerSyncState, TEXT("r.Streaming.SyncStatesWhenBlocking"), 1, bOverrideValues);
 	}
 
+	MOVIEPIPELINE_STORE_AND_OVERRIDE_CVAR_BOOL(PreviousAlphaOutput, TEXT("r.PostProcessing.PropagateAlpha"), UE::MoviePipeline::GetAlphaOutputOverride(), bOverrideValues);
+
 	// Set cvars that allow MRQ to run correctly. These are not exposed as properties on the node, but the chance of a
 	// user needing to customize these (or be aware of them) is very low.
 	{
diff --git a/Engine/Plugins/MovieScene/MovieRenderPipeline/Source/MovieRenderPipelineCore/Private/MoviePipeline.cpp b/Engine/Plugins/MovieScene/MovieRenderPipeline/Source/MovieRenderPipelineCore/Private/MoviePipeline.cpp
--- a/Engine/Plugins/MovieScene/MovieRenderPipeline/Source/MovieRenderPipelineCore/Private/MoviePipeline.cpp
+++ b/Engine/Plugins/MovieScene/MovieRenderPipeline/Source/MovieRenderPipelineCore/Private/MoviePipeline.cpp
@@ -101,25 +101,6 @@ void UMoviePipeline::ValidateSequenceAndSettings() const
 {
 	// ToDo: 
 	// Warn for Blueprint Streaming Levels
-
-	// Check to see if they're trying to output alpha and don't have the required project setting set.
-	{
-		IConsoleVariable* TonemapAlphaCVar = IConsoleManager::Get().FindConsoleVariable(TEXT("r.PostProcessing.PropagateAlpha"));
-		check(TonemapAlphaCVar);
-
-		TArray<UMoviePipelineRenderPass*> OutputSettings = GetPipelinePrimaryConfig()->FindSettings<UMoviePipelineRenderPass>();
-		bool bAnyOutputWantsAlpha = false;
-
-		for (const UMoviePipelineRenderPass* Output : OutputSettings)
-		{
-			bAnyOutputWantsAlpha |= Output->IsAlphaInTonemapperRequired();
-		}
-
-		if (bAnyOutputWantsAlpha && TonemapAlphaCVar->GetBool() == false)
-		{
-			UE_LOG(LogMovieRenderPipeline, Warning, TEXT("An output requested Alpha Support but the required project setting is not enabled! Go to Project Settings > Rendering > PostProcessing > 'Alpha Output' and set it to 'True'."));
-		}
-	}
 }
 
 void UMoviePipeline::Initialize(UMoviePipelineExecutorJob* InJob)
diff --git a/Engine/Plugins/MovieScene/MovieRenderPipeline/Source/MovieRenderPipelineCore/Private/MoviePipelineGameOverrideSetting.cpp b/Engine/Plugins/MovieScene/MovieRenderPipeline/Source/MovieRenderPipelineCore/Private/MoviePipelineGameOverrideSetting.cpp
--- a/Engine/Plugins/MovieScene/MovieRenderPipeline/Source/MovieRenderPipelineCore/Private/MoviePipelineGameOverrideSetting.cpp
+++ b/Engine/Plugins/MovieScene/MovieRenderPipeline/Source/MovieRenderPipelineCore/Private/MoviePipelineGameOverrideSetting.cpp
@@ -121,6 +121,8 @@ void UMoviePipelineGameOverrideSetting::ApplyCVarSettings(const bool bOverrideVa
 	{
 		MOVIEPIPELINE_STORE_AND_OVERRIDE_CVAR_INT(PreviousStreamingManagerSyncState, TEXT("r.Streaming.SyncStatesWhenBlocking"), 1, bOverrideValues);
 	}
+
+	MOVIEPIPELINE_STORE_AND_OVERRIDE_CVAR_BOOL(PreviousAlphaOutput, TEXT("r.PostProcessing.PropagateAlpha"), UE::MoviePipeline::GetAlphaOutputOverride(), bOverrideValues);
 	
 #if WITH_EDITOR
 	// To make sure the GeometryCache streamer doesn't skip frames and doesn't pop up notification during rendering
@@ -276,4 +278,6 @@ void UMoviePipelineGameOverrideSetting::BuildNewProcessCommandLineArgsImpl(TArra
 	{
 		InOutDeviceProfileCvars.AddUnique(FString::Printf(TEXT("r.Nanite.VSMInvalidateOnLODDelta=%d"), 1));
 	}
+
+	InOutDeviceProfileCvars.AddUnique(FString::Printf(TEXT("r.PostProcessing.PropagateAlpha=%d"), UE::MoviePipeline::GetAlphaOutputOverride() ? 1 : 0));
 }
diff --git a/Engine/Plugins/MovieScene/MovieRenderPipeline/Source/MovieRenderPipelineCore/Private/MoviePipelineRenderPass.cpp b/Engine/Plugins/MovieScene/MovieRenderPipeline/Source/MovieRenderPipelineCore/Private/MoviePipelineRenderPass.cpp
deleted file mode 100644
diff --git a/Engine/Plugins/MovieScene/MovieRenderPipeline/Source/MovieRenderPipelineCore/Private/MoviePipelineUtils.cpp b/Engine/Plugins/MovieScene/MovieRenderPipeline/Source/MovieRenderPipelineCore/Private/MoviePipelineUtils.cpp
--- a/Engine/Plugins/MovieScene/MovieRenderPipeline/Source/MovieRenderPipelineCore/Private/MoviePipelineUtils.cpp
+++ b/Engine/Plugins/MovieScene/MovieRenderPipeline/Source/MovieRenderPipelineCore/Private/MoviePipelineUtils.cpp
@@ -44,6 +44,12 @@
 #include "Components/SkeletalMeshComponent.h"
 #include "Graph/Nodes/MovieGraphSubgraphNode.h"
 
+static TAutoConsoleVariable<bool> CVarMoviePipelineAlphaOutputOverride(
+	TEXT("MoviePipeline.AlphaOutputOverride"),
+	true,
+	TEXT("Always overrides the \"Alpha Output\" renderer project setting by controlling r.PostProcessing.PropagateAlpha during movie pipeline renders.\n"),
+	ECVF_Default);
+
 namespace UE
 {
 	namespace MovieRenderPipeline
@@ -833,6 +839,11 @@ namespace UE
 		DECLARE_CYCLE_STAT(TEXT("STAT_MoviePipeline_HardwareMetadata"), STAT_HardwareMetadata, STATGROUP_MoviePipeline);
 		DECLARE_CYCLE_STAT(TEXT("STAT_MoviePipeline_ClothAdjust"), STAT_ClothSubstepAdjust, STATGROUP_MoviePipeline);
 
+		bool GetAlphaOutputOverride()
+		{
+			return CVarMoviePipelineAlphaOutputOverride.GetValueOnAnyThread();
+		}
+
 		void ConformOutputFormatStringToken(FString& InOutFilenameFormatString, const FStringView InToken, const FName& InNodeName, const FName& InBranchName)
 		{
 			static const FString FrameNumberIdentifiers[] = { TEXT("{frame_number}"), TEXT("{frame_number_shot}"), TEXT("{frame_number_rel}"), TEXT("{frame_number_shot_rel}") };
diff --git a/Engine/Plugins/MovieScene/MovieRenderPipeline/Source/MovieRenderPipelineCore/Public/Graph/Nodes/MovieGraphGlobalGameOverrides.h b/Engine/Plugins/MovieScene/MovieRenderPipeline/Source/MovieRenderPipelineCore/Public/Graph/Nodes/MovieGraphGlobalGameOverrides.h
--- a/Engine/Plugins/MovieScene/MovieRenderPipeline/Source/MovieRenderPipelineCore/Public/Graph/Nodes/MovieGraphGlobalGameOverrides.h
+++ b/Engine/Plugins/MovieScene/MovieRenderPipeline/Source/MovieRenderPipelineCore/Public/Graph/Nodes/MovieGraphGlobalGameOverrides.h
@@ -220,6 +220,7 @@ private:
 	int32 PreviousChaosClothUseTimeStepSmoothing;
 	int32 PreviousSkipWaterInfoTextureRenderWhenWorldRenderingDisabled;
 	int32 PreviousNaniteVSMInvalidateOnLODDelta;
+	int32 PreviousAlphaOutput;
 
 #if WITH_EDITOR
 	int32 PreviousGeoCacheStreamerShowNotification;
diff --git a/Engine/Plugins/MovieScene/MovieRenderPipeline/Source/MovieRenderPipelineCore/Public/MoviePipelineGameOverrideSetting.h b/Engine/Plugins/MovieScene/MovieRenderPipeline/Source/MovieRenderPipelineCore/Public/MoviePipelineGameOverrideSetting.h
--- a/Engine/Plugins/MovieScene/MovieRenderPipeline/Source/MovieRenderPipelineCore/Public/MoviePipelineGameOverrideSetting.h
+++ b/Engine/Plugins/MovieScene/MovieRenderPipeline/Source/MovieRenderPipelineCore/Public/MoviePipelineGameOverrideSetting.h
@@ -162,4 +162,5 @@ private:
 	int32 PreviousChaosClothUseTimeStepSmoothing;
 	int32 PreviousSkipWaterInfoTextureRenderWhenWorldRenderingDisabled;
 	int32 PreviousNaniteVSMInvalidateOnLODDelta;
+	bool PreviousAlphaOutput;
 };
\ No newline at end of file
diff --git a/Engine/Plugins/MovieScene/MovieRenderPipeline/Source/MovieRenderPipelineCore/Public/MoviePipelineRenderPass.h b/Engine/Plugins/MovieScene/MovieRenderPipeline/Source/MovieRenderPipelineCore/Public/MoviePipelineRenderPass.h
--- a/Engine/Plugins/MovieScene/MovieRenderPipeline/Source/MovieRenderPipelineCore/Public/MoviePipelineRenderPass.h
+++ b/Engine/Plugins/MovieScene/MovieRenderPipeline/Source/MovieRenderPipelineCore/Public/MoviePipelineRenderPass.h
@@ -12,9 +12,6 @@ class MOVIERENDERPIPELINECORE_API UMoviePipelineRenderPass : public UMoviePipeli
 {
 	GENERATED_BODY()
 public:
-	// UMoviePipelineSetting Interface
-	virtual void ValidateStateImpl() override;
-	// ~UMoviePipelineSetting Interface
 	
 	void Setup(const MoviePipeline::FMoviePipelineRenderPassInitSettings& InPassInitSettings)
 	{
@@ -45,9 +42,10 @@ public:
 		RenderSample_GameThreadImpl(InSampleState);
 	}
 
+	UE_DEPRECATED(5.6, "This function is no longer in use.")
 	bool IsAlphaInTonemapperRequired() const
 	{
-		return IsAlphaInTonemapperRequiredImpl();
+		return false;
 	}
 
 	virtual bool NeedsFrameThrottle() const { return false; }
@@ -66,5 +64,4 @@ protected:
 	virtual void OnFrameStartImpl() {}
 	virtual void GatherOutputPassesImpl(TArray<FMoviePipelinePassIdentifier>& ExpectedRenderPasses) {}
 	virtual void RenderSample_GameThreadImpl(const FMoviePipelineRenderPassMetrics& InSampleState) {}
-	virtual bool IsAlphaInTonemapperRequiredImpl() const { return false; }
 };
\ No newline at end of file
diff --git a/Engine/Plugins/MovieScene/MovieRenderPipeline/Source/MovieRenderPipelineCore/Public/MoviePipelineUtils.h b/Engine/Plugins/MovieScene/MovieRenderPipeline/Source/MovieRenderPipelineCore/Public/MoviePipelineUtils.h
--- a/Engine/Plugins/MovieScene/MovieRenderPipeline/Source/MovieRenderPipelineCore/Public/MoviePipelineUtils.h
+++ b/Engine/Plugins/MovieScene/MovieRenderPipeline/Source/MovieRenderPipelineCore/Public/MoviePipelineUtils.h
@@ -50,6 +50,23 @@ MOVIERENDERPIPELINECORE_API void GetPassCompositeData(FMoviePipelineMergerOutput
 
 }
 
+#define MOVIEPIPELINE_STORE_AND_OVERRIDE_CVAR_BOOL(InOutVariable, CVarName, OverrideValue, bUseOverride) \
+{ \
+	IConsoleVariable* CVar = IConsoleManager::Get().FindConsoleVariable(CVarName); \
+	if(ensureMsgf(CVar, TEXT("Failed to find CVar " #CVarName " to override."))) \
+	{ \
+		if(bUseOverride) \
+		{ \
+			InOutVariable = CVar->GetBool(); \
+			CVar->SetWithCurrentPriority(OverrideValue); \
+		} \
+		else \
+		{ \
+			CVar->SetWithCurrentPriority(InOutVariable); \
+		} \
+	} \
+}
+
 #define MOVIEPIPELINE_STORE_AND_OVERRIDE_CVAR_INT(InOutVariable, CVarName, OverrideValue, bUseOverride) \
 { \
 	IConsoleVariable* CVar = IConsoleManager::Get().FindConsoleVariable(CVarName); \
@@ -84,6 +101,23 @@ MOVIERENDERPIPELINECORE_API void GetPassCompositeData(FMoviePipelineMergerOutput
 	} \
 }
 
+#define MOVIEPIPELINE_STORE_AND_OVERRIDE_CVAR_BOOL_IF_EXIST(InOutVariable, CVarName, OverrideValue, bUseOverride) \
+{ \
+	IConsoleVariable* CVar = IConsoleManager::Get().FindConsoleVariable(CVarName); \
+	if(CVar) \
+	{ \
+		if(bUseOverride) \
+		{ \
+			InOutVariable = CVar->GetBool(); \
+			CVar->SetWithCurrentPriority(OverrideValue); \
+		} \
+		else \
+		{ \
+			CVar->SetWithCurrentPriority(InOutVariable); \
+		} \
+	} \
+}
+
 #define MOVIEPIPELINE_STORE_AND_OVERRIDE_CVAR_INT_IF_EXIST(InOutVariable, CVarName, OverrideValue, bUseOverride) \
 { \
 	IConsoleVariable* CVar = IConsoleManager::Get().FindConsoleVariable(CVarName); \
@@ -138,6 +172,8 @@ namespace UE
 
 	namespace MoviePipeline
 	{
+		/** Get the value of the MovieRenderPipeline.AlphaOutputOverride console variable. */
+		MOVIERENDERPIPELINECORE_API bool GetAlphaOutputOverride();
 		MOVIERENDERPIPELINECORE_API void ConformOutputFormatStringToken(FString& InOutFilenameFormatString, const FStringView InToken, const FName& InNodeName, const FName& InBranchName);
 		MOVIERENDERPIPELINECORE_API void ValidateOutputFormatString(FString& InOutFilenameFormatString, const bool bTestRenderPass, const bool bTestFrameNumber, const bool bIncludeCameraName = false);
 		MOVIERENDERPIPELINECORE_API void RemoveFrameNumberFormatStrings(FString& InOutFilenameFormatString, const bool bIncludeShots);
diff --git a/Engine/Plugins/MovieScene/MovieRenderPipeline/Source/MovieRenderPipelineRenderPasses/Public/MoviePipelineDeferredPasses.h b/Engine/Plugins/MovieScene/MovieRenderPipeline/Source/MovieRenderPipelineRenderPasses/Public/MoviePipelineDeferredPasses.h
--- a/Engine/Plugins/MovieScene/MovieRenderPipeline/Source/MovieRenderPipelineRenderPasses/Public/MoviePipelineDeferredPasses.h
+++ b/Engine/Plugins/MovieScene/MovieRenderPipeline/Source/MovieRenderPipelineRenderPasses/Public/MoviePipelineDeferredPasses.h
@@ -66,7 +66,6 @@ protected:
 	virtual void GatherOutputPassesImpl(TArray<FMoviePipelinePassIdentifier>& ExpectedRenderPasses) override;
 	virtual bool IsAntiAliasingSupported() const { return !bDisableMultisampleEffects; }
 	virtual int32 GetOutputFileSortingOrder() const override { return 0; }
-	virtual bool IsAlphaInTonemapperRequiredImpl() const override { return bAccumulatorIncludesAlpha; }
 	virtual FSceneViewStateInterface* GetSceneViewStateInterface(IViewCalcPayload* OptPayload = nullptr) override;
 	virtual FSceneViewStateInterface* GetExposureSceneViewStateInterface(IViewCalcPayload* OptPayload = nullptr, int32 CubeFaceIndex = INDEX_NONE) override;
 	virtual void AddViewExtensions(FSceneViewFamilyContext& InContext, FMoviePipelineRenderPassMetrics& InOutSampleState) override;
@@ -111,9 +110,8 @@ private:
 
 public:
 	/**
-	* Should multiple temporal/spatial samples accumulate the alpha channel? This requires r.PostProcessing.PropagateAlpha
-	* to be enabled (see "Alpha Output" under Project Settings > Rendering). This adds
-	* ~30% cost to the accumulation so you should not enable it unless necessary. You must delete both the sky and fog to ensure
+	* Should multiple temporal/spatial samples accumulate the alpha channel? This adds ~30% cost to the
+	* accumulation so you should not enable it unless necessary. You must delete both the sky and fog to ensure
 	* that they do not make all pixels opaque.
 	*/
 	UPROPERTY(EditAnywhere, BlueprintReadWrite, Category = "Settings")
