MoviePipeline: Automation: Fixed transient settings not being initialized during tests, which meant the Game Overrides setting was not being applied. Set bTrackFrequentCalls to false when looking for cvars to avoid warnings during tests for optional cvars.

Increased tolerance on TSR/TAA test due to continued non-determinism. Updated ground truth images.

#virtualized

diff --git a/Engine/Plugins/MovieScene/MovieRenderPipeline/Source/MovieRenderPipelineCore/Public/MoviePipelineUtils.h b/Engine/Plugins/MovieScene/MovieRenderPipeline/Source/MovieRenderPipelineCore/Public/MoviePipelineUtils.h
--- a/Engine/Plugins/MovieScene/MovieRenderPipeline/Source/MovieRenderPipelineCore/Public/MoviePipelineUtils.h
+++ b/Engine/Plugins/MovieScene/MovieRenderPipeline/Source/MovieRenderPipelineCore/Public/MoviePipelineUtils.h
@@ -52,7 +52,8 @@ MOVIERENDERPIPELINECORE_API void GetPassCompositeData(FMoviePipelineMergerOutput
 
 #define MOVIEPIPELINE_STORE_AND_OVERRIDE_CVAR_BOOL(InOutVariable, CVarName, OverrideValue, bUseOverride) \
 { \
-	IConsoleVariable* CVar = IConsoleManager::Get().FindConsoleVariable(CVarName); \
+	const bool bTrackFrequentCalls = false; \
+	IConsoleVariable* CVar = IConsoleManager::Get().FindConsoleVariable(CVarName, bTrackFrequentCalls); \
 	if(ensureMsgf(CVar, TEXT("Failed to find CVar " #CVarName " to override."))) \
 	{ \
 		if(bUseOverride) \
@@ -69,7 +70,8 @@ MOVIERENDERPIPELINECORE_API void GetPassCompositeData(FMoviePipelineMergerOutput
 
 #define MOVIEPIPELINE_STORE_AND_OVERRIDE_CVAR_INT(InOutVariable, CVarName, OverrideValue, bUseOverride) \
 { \
-	IConsoleVariable* CVar = IConsoleManager::Get().FindConsoleVariable(CVarName); \
+	const bool bTrackFrequentCalls = false; \
+	IConsoleVariable* CVar = IConsoleManager::Get().FindConsoleVariable(CVarName, bTrackFrequentCalls); \
 	if(ensureMsgf(CVar, TEXT("Failed to find CVar " #CVarName " to override."))) \
 	{ \
 		if(bUseOverride) \
@@ -86,7 +88,8 @@ MOVIERENDERPIPELINECORE_API void GetPassCompositeData(FMoviePipelineMergerOutput
 
 #define MOVIEPIPELINE_STORE_AND_OVERRIDE_CVAR_FLOAT(InOutVariable, CVarName, OverrideValue, bUseOverride) \
 { \
-	IConsoleVariable* CVar = IConsoleManager::Get().FindConsoleVariable(CVarName); \
+	const bool bTrackFrequentCalls = false; \
+	IConsoleVariable* CVar = IConsoleManager::Get().FindConsoleVariable(CVarName, bTrackFrequentCalls); \
 	if(ensureMsgf(CVar, TEXT("Failed to find CVar " #CVarName " to override."))) \
 	{ \
 		if(bUseOverride) \
@@ -103,7 +106,8 @@ MOVIERENDERPIPELINECORE_API void GetPassCompositeData(FMoviePipelineMergerOutput
 
 #define MOVIEPIPELINE_STORE_AND_OVERRIDE_CVAR_BOOL_IF_EXIST(InOutVariable, CVarName, OverrideValue, bUseOverride) \
 { \
-	IConsoleVariable* CVar = IConsoleManager::Get().FindConsoleVariable(CVarName); \
+	const bool bTrackFrequentCalls = false; \
+	IConsoleVariable* CVar = IConsoleManager::Get().FindConsoleVariable(CVarName, bTrackFrequentCalls); \
 	if(CVar) \
 	{ \
 		if(bUseOverride) \
@@ -120,7 +124,8 @@ MOVIERENDERPIPELINECORE_API void GetPassCompositeData(FMoviePipelineMergerOutput
 
 #define MOVIEPIPELINE_STORE_AND_OVERRIDE_CVAR_INT_IF_EXIST(InOutVariable, CVarName, OverrideValue, bUseOverride) \
 { \
-	IConsoleVariable* CVar = IConsoleManager::Get().FindConsoleVariable(CVarName); \
+	const bool bTrackFrequentCalls = false; \
+	IConsoleVariable* CVar = IConsoleManager::Get().FindConsoleVariable(CVarName, bTrackFrequentCalls); \
 	if(CVar) \
 	{ \
 		if(bUseOverride) \
@@ -137,7 +142,8 @@ MOVIERENDERPIPELINECORE_API void GetPassCompositeData(FMoviePipelineMergerOutput
 
 #define MOVIEPIPELINE_STORE_AND_OVERRIDE_CVAR_FLOAT_IF_EXIST(InOutVariable, CVarName, OverrideValue, bUseOverride) \
 { \
-	IConsoleVariable* CVar = IConsoleManager::Get().FindConsoleVariable(CVarName); \
+	const bool bTrackFrequentCalls = false; \
+	IConsoleVariable* CVar = IConsoleManager::Get().FindConsoleVariable(CVarName, bTrackFrequentCalls); \
 	if(CVar) \
 	{ \
 		if(bUseOverride) \
diff --git a/Engine/Plugins/MovieScene/MovieRenderPipeline/Source/MovieRenderPipelineEditor/Private/MoviePipelineFunctionalTestBase.cpp b/Engine/Plugins/MovieScene/MovieRenderPipeline/Source/MovieRenderPipelineEditor/Private/MoviePipelineFunctionalTestBase.cpp
--- a/Engine/Plugins/MovieScene/MovieRenderPipeline/Source/MovieRenderPipelineEditor/Private/MoviePipelineFunctionalTestBase.cpp
+++ b/Engine/Plugins/MovieScene/MovieRenderPipeline/Source/MovieRenderPipelineEditor/Private/MoviePipelineFunctionalTestBase.cpp
@@ -64,6 +64,9 @@ void AMoviePipelineFunctionalTestBase::StartTest()
 	ActiveQueue = NewObject<UMoviePipelineQueue>(GetWorld());
 	ActiveQueue->CopyFrom(QueuePreset);
 
+	// Ensure we've initialized any transient settings (ie: the game overrides setting that is automatically added),
+	// otherwise it won't get called.
+	ActiveQueue->GetJobs()[0]->GetConfiguration()->InitializeTransientSettings();
 	ActiveMoviePipeline->Initialize(ActiveQueue->GetJobs()[0]); // Prepare Test will ensure we have one job
 }
 
