MoviePipeline: Added ability to show editor-only actors in a Quick Render.

The primary use case for this is to be able to see animation-related actors, like cameras, camera rig rails, camera cranes, camera frustums, and control rig gizmos.

diff --git a/Engine/Plugins/MovieScene/MovieRenderPipeline/Source/MovieRenderPipelineCore/Private/Graph/MovieGraphRenderLayerSubsystem.cpp b/Engine/Plugins/MovieScene/MovieRenderPipeline/Source/MovieRenderPipelineCore/Private/Graph/MovieGraphRenderLayerSubsystem.cpp
--- a/Engine/Plugins/MovieScene/MovieRenderPipeline/Source/MovieRenderPipelineCore/Private/Graph/MovieGraphRenderLayerSubsystem.cpp
+++ b/Engine/Plugins/MovieScene/MovieRenderPipeline/Source/MovieRenderPipelineCore/Private/Graph/MovieGraphRenderLayerSubsystem.cpp
@@ -176,6 +176,7 @@ UMovieGraphRenderPropertyModifier::UMovieGraphRenderPropertyModifier()
 	, bCastShadowWhileHidden(false)
 	, bAffectIndirectLightingWhileHidden(false)
 	, bHoldout(false)
+	, bProcessEditorOnlyActors(false)
 {
 	// Note: The default modifier values here reflect the defaults on the scene component. If a modifier property is marked as overridden, the
 	// override will initially be a no-op due to the defaults being the same.
@@ -205,8 +206,8 @@ void UMovieGraphRenderPropertyModifier::ApplyModifier(const UWorld* World)
 	auto CacheComponentState = [this](USceneComponent* InComponent)
 	{
 #if WITH_EDITORONLY_DATA
-		// Don't bother processing editor-only components (editor billboard icons, text, etc)
-		if (InComponent->IsEditorOnly())
+		// Don't bother processing editor-only components (editor billboard icons, text, etc). This behavior can be optionally ignored for niche use cases.
+		if (InComponent->IsEditorOnly() && !bProcessEditorOnlyActors)
 		{
 			return;
 		}
diff --git a/Engine/Plugins/MovieScene/MovieRenderPipeline/Source/MovieRenderPipelineCore/Public/Graph/MovieGraphQuickRenderSettings.h b/Engine/Plugins/MovieScene/MovieRenderPipeline/Source/MovieRenderPipelineCore/Public/Graph/MovieGraphQuickRenderSettings.h
--- a/Engine/Plugins/MovieScene/MovieRenderPipeline/Source/MovieRenderPipelineCore/Public/Graph/MovieGraphQuickRenderSettings.h
+++ b/Engine/Plugins/MovieScene/MovieRenderPipeline/Source/MovieRenderPipelineCore/Public/Graph/MovieGraphQuickRenderSettings.h
@@ -77,6 +77,7 @@ enum class EMovieGraphQuickRenderViewportLookFlags : uint8
 	ShowFlags	= 1 << 1,
 	ViewMode	= 1 << 2,
 	Visibility	= 1 << 3,
+	EditorOnlyActors = 1 << 4 UMETA(DisplayName = "Show Editor-Only Actors")
 };
 ENUM_CLASS_FLAGS(EMovieGraphQuickRenderViewportLookFlags)
 
diff --git a/Engine/Plugins/MovieScene/MovieRenderPipeline/Source/MovieRenderPipelineCore/Public/Graph/MovieGraphRenderLayerSubsystem.h b/Engine/Plugins/MovieScene/MovieRenderPipeline/Source/MovieRenderPipelineCore/Public/Graph/MovieGraphRenderLayerSubsystem.h
--- a/Engine/Plugins/MovieScene/MovieRenderPipeline/Source/MovieRenderPipelineCore/Public/Graph/MovieGraphRenderLayerSubsystem.h
+++ b/Engine/Plugins/MovieScene/MovieRenderPipeline/Source/MovieRenderPipelineCore/Public/Graph/MovieGraphRenderLayerSubsystem.h
@@ -1018,6 +1018,9 @@ public:
 
 	UPROPERTY(EditAnywhere, BlueprintReadWrite, Category = Overrides, meta = (InlineEditConditionToggle))
 	uint8 bOverride_bHoldout : 1;
+
+	UPROPERTY(EditAnywhere, BlueprintReadWrite, Category = Overrides, meta = (InlineEditConditionToggle))
+	uint8 bOverride_bProcessEditorOnlyActors : 1;
 	
 	/**
 	 * If true, the actor will not be visible and will not contribute to any secondary effects (shadows, indirect
@@ -1044,6 +1047,15 @@ public:
 	 */
 	UPROPERTY(BlueprintReadWrite, EditAnywhere, Category = "Settings", meta = (EditCondition = "bOverride_bHoldout"))
 	uint8 bHoldout : 1;
+
+	/**
+	 * If true, this modifier will process editor-only actors (eg, the billboard icons that represent lights). Defaults to false. Generally, processing
+	 * editor-only actors is a waste of cycles because they will never be shown in a render, and there's no point in messing with them. However there
+	 * are some niche cases where they need to be modified (eg, Quick Render). This is not a property exposed to the UI or saved because it's only
+	 * meant to be used temporarily (like by scripting; note that bOverride_bProcessEditorOnlyActors still needs to be set).
+	 */
+	UPROPERTY(BlueprintReadWrite, Category = "Settings")
+	uint8 bProcessEditorOnlyActors : 1;
 };
 
 /**
diff --git a/Engine/Plugins/MovieScene/MovieRenderPipeline/Source/MovieRenderPipelineEditor/Private/Graph/MovieGraphQuickRender.cpp b/Engine/Plugins/MovieScene/MovieRenderPipeline/Source/MovieRenderPipelineEditor/Private/Graph/MovieGraphQuickRender.cpp
--- a/Engine/Plugins/MovieScene/MovieRenderPipeline/Source/MovieRenderPipelineEditor/Private/Graph/MovieGraphQuickRender.cpp
+++ b/Engine/Plugins/MovieScene/MovieRenderPipeline/Source/MovieRenderPipelineEditor/Private/Graph/MovieGraphQuickRender.cpp
@@ -375,6 +375,12 @@ void UMovieGraphQuickRenderSubsystem::ApplyQuickRenderUpdatesToDuplicateGraph(co
 		ApplyQuickRenderUpdatesToDuplicateGraph_ApplyEditorVisibility(InEditorWorld);
 	}
 
+	// Apply editor-only actor visibility. Normally editor-only actors won't be processed by the MRG modifier system.
+	if (IsViewportLookFlagActive(EMovieGraphQuickRenderViewportLookFlags::EditorOnlyActors))
+	{
+		ApplyQuickRenderUpdatesToDuplicateGraph_ApplyEditorOnlyActorVisibility(InEditorWorld);
+	}
+
 	// Apply OCIO if it is activated on the viewport and enabled in the viewport look flags
 	if (IsViewportLookFlagActive(EMovieGraphQuickRenderViewportLookFlags::Ocio))
 	{
@@ -431,35 +437,81 @@ void UMovieGraphQuickRenderSubsystem::ApplyQuickRenderUpdatesToDuplicateGraph_Ap
 		}
 	}
 
-	auto AddNewCollectionWithVisibilityModifier = [this]<typename QueryType>(const FString& InOperationName, const bool bModifierShouldHide) -> QueryType*
+	UMovieGraphConditionGroupQuery_Actor* VisibleActorQuery =
+		AddNewCollectionWithVisibilityModifier<UMovieGraphConditionGroupQuery_Actor>(TEXT("VISIBLE_ACTORS"), false);
+	VisibleActorQuery->ActorsAndComponentsToMatch = VisibleActors;
+	
+	UMovieGraphConditionGroupQuery_Actor* HiddenActorQuery =
+		AddNewCollectionWithVisibilityModifier<UMovieGraphConditionGroupQuery_Actor>(TEXT("HIDDEN_ACTORS"), true);
+	HiddenActorQuery->ActorsAndComponentsToMatch = HiddenActors;
+}
+
+void UMovieGraphQuickRenderSubsystem::ApplyQuickRenderUpdatesToDuplicateGraph_ApplyEditorOnlyActorVisibility(const UWorld* InEditorWorld)
+{
+	TArray<FMovieGraphActorQueryEntry> VisibleActors;
+	
+	// Determine the editor-only actors that should be made visible
+	for (FActorIterator It(InEditorWorld); It; ++It)
 	{
-		const UMovieGraphCollectionNode* NewCollectionNode = Cast<UMovieGraphCollectionNode>(InjectNodeIntoBranch(UMovieGraphCollectionNode::StaticClass(), UMovieGraphNode::GlobalsPinName));
+		AActor* Actor = *It;
 
-		const FString NewCollectionName(TEXT("__AUTOGEN_COLLECTION_") + InOperationName);
-		NewCollectionNode->Collection->SetCollectionName(NewCollectionName);
+		if (!Actor)
+		{
+			continue;
+		}
 
-		QueryType* NewConditionGroupQuery = Cast<QueryType>(NewCollectionNode->Collection->AddConditionGroup()->AddQuery(QueryType::StaticClass()));
+		bool bHasEditorOnlyComponent = false;
+		
+		for (const UActorComponent* ActorComponent : Actor->GetComponents())
+		{
+			if (ActorComponent->IsEditorOnly())
+			{
+				bHasEditorOnlyComponent = true;
+				break;
+			}
+		}
 
-		UMovieGraphModifierNode* NewModifierNode = Cast<UMovieGraphModifierNode>(InjectNodeIntoBranch(UMovieGraphModifierNode::StaticClass(), UMovieGraphNode::GlobalsPinName));
-		NewModifierNode->ModifierName = TEXT("__AUTOGEN_MODIFIER_") + InOperationName;
-		NewModifierNode->AddCollection(FName(NewCollectionName));
+		if (Actor->IsEditorOnly() || bHasEditorOnlyComponent)
+		{
+			FMovieGraphActorQueryEntry& VisibleActorEntry = VisibleActors.AddDefaulted_GetRef();
+			VisibleActorEntry.ActorToMatch = Actor;
+		}
 
-		UMovieGraphRenderPropertyModifier* RenderPropertyModifier = Cast<UMovieGraphRenderPropertyModifier>(NewModifierNode->GetModifier(UMovieGraphRenderPropertyModifier::StaticClass()));
+		// The control rig actor class name is referenced by name here in order to prevent a dependency.
+		static const FName ControlRigShapeActorName(TEXT("ControlRigShapeActor"));
 
-		RenderPropertyModifier->bOverride_bIsHidden = true;
-		RenderPropertyModifier->bIsHidden = bModifierShouldHide;
+		// The Control Rig manipulation gizmos are important to show in the editor-only actor mode, but they will not be copied into the PIE world
+		// because they're marked as Transient. To work around this, we can temporarily mark them as RF_NonPIEDuplicateTransient which will inform
+		// the serialization process that it's OK to copy these actors to PIE.
+		if (Actor->HasAnyFlags(RF_Transient) && (Actor->GetClass()->GetName() == ControlRigShapeActorName))
+		{
+			// Since we need to un-set this flag after the render is done, only apply RF_NonPIEDuplicateTransient if it's not already applied (ie,
+			// we don't want to remove RF_NonPIEDuplicateTransient after the render finishes if Quick Render didn't originally apply it).
+			if (!Actor->HasAnyFlags(RF_NonPIEDuplicateTransient))
+			{
+				CachedPrePieData.TemporarilyNonTransientActors.Add(MakeWeakObjectPtr(Actor));
+				Actor->SetFlags(RF_NonPIEDuplicateTransient);
+			}
 
-		// The caller can modify the returned condition group as needed
-		return NewConditionGroupQuery;
-	};
+			for (UActorComponent* Component : Actor->GetComponents())
+			{
+				if (Component->HasAnyFlags(RF_Transient))
+				{
+					if (!Component->HasAnyFlags(RF_NonPIEDuplicateTransient))
+					{
+						CachedPrePieData.TemporarilyNonTransientComponents.Add(MakeWeakObjectPtr(Component));
+						Component->SetFlags(RF_NonPIEDuplicateTransient);
+					}
+				}
+			}
+		}
+	}
 
+	constexpr bool bModifierShouldHide = false;
+	constexpr bool bProcessEditorOnlyActors = true;
 	UMovieGraphConditionGroupQuery_Actor* VisibleActorQuery =
-		AddNewCollectionWithVisibilityModifier.operator()<UMovieGraphConditionGroupQuery_Actor>(TEXT("VISIBLE_ACTORS"), false);
+		AddNewCollectionWithVisibilityModifier<UMovieGraphConditionGroupQuery_Actor>(TEXT("VISIBLE_EDITORONLY_ACTORS"), bModifierShouldHide, bProcessEditorOnlyActors);
 	VisibleActorQuery->ActorsAndComponentsToMatch = VisibleActors;
-	
-	UMovieGraphConditionGroupQuery_Actor* HiddenActorQuery =
-		AddNewCollectionWithVisibilityModifier.operator()<UMovieGraphConditionGroupQuery_Actor>(TEXT("HIDDEN_ACTORS"), true);
-	HiddenActorQuery->ActorsAndComponentsToMatch = HiddenActors;
 }
 
 void UMovieGraphQuickRenderSubsystem::ApplyQuickRenderUpdatesToDuplicateGraph_ApplyOcio(const FOpenColorIODisplayConfiguration* InOcioConfiguration)
@@ -589,6 +641,35 @@ void UMovieGraphQuickRenderSubsystem::ApplyQuickRenderUpdatesToDuplicateGraph_Sc
 	}
 }
 
+template <typename QueryType>
+QueryType* UMovieGraphQuickRenderSubsystem::AddNewCollectionWithVisibilityModifier(const FString& InOperationName, const bool bModifierShouldHide, const bool bProcessEditorOnlyActors)
+{
+	const UMovieGraphCollectionNode* NewCollectionNode = Cast<UMovieGraphCollectionNode>(InjectNodeIntoBranch(UMovieGraphCollectionNode::StaticClass(), UMovieGraphNode::GlobalsPinName));
+
+	const FString NewCollectionName(TEXT("__AUTOGEN_COLLECTION_") + InOperationName);
+	NewCollectionNode->Collection->SetCollectionName(NewCollectionName);
+
+	QueryType* NewConditionGroupQuery = Cast<QueryType>(NewCollectionNode->Collection->AddConditionGroup()->AddQuery(QueryType::StaticClass()));
+
+	UMovieGraphModifierNode* NewModifierNode = Cast<UMovieGraphModifierNode>(InjectNodeIntoBranch(UMovieGraphModifierNode::StaticClass(), UMovieGraphNode::GlobalsPinName));
+	NewModifierNode->ModifierName = TEXT("__AUTOGEN_MODIFIER_") + InOperationName;
+	NewModifierNode->AddCollection(FName(NewCollectionName));
+
+	UMovieGraphRenderPropertyModifier* RenderPropertyModifier = Cast<UMovieGraphRenderPropertyModifier>(NewModifierNode->GetModifier(UMovieGraphRenderPropertyModifier::StaticClass()));
+
+	RenderPropertyModifier->bOverride_bIsHidden = true;
+	RenderPropertyModifier->bIsHidden = bModifierShouldHide;
+
+	if (bProcessEditorOnlyActors)
+	{
+		RenderPropertyModifier->bOverride_bProcessEditorOnlyActors = true;
+		RenderPropertyModifier->bProcessEditorOnlyActors = true;
+	}
+
+	// The caller can modify the returned condition group as needed
+	return NewConditionGroupQuery;
+}
+
 UMovieGraphNode* UMovieGraphQuickRenderSubsystem::InjectNodeIntoBranch(const TSubclassOf<UMovieGraphNode> NodeType, const FName& InBranchName) const
 {
 	check(IsValid(TemporaryGraph));
@@ -692,6 +773,26 @@ void UMovieGraphQuickRenderSubsystem::RestorePreRenderState()
 	{
 		RenderingLevelSequence->GetMovieScene()->RemoveTrack(*TrackToRemove);
 	}
+
+	// Strip any transient flags off of relevant components and actors
+	for (const TWeakObjectPtr<UActorComponent>& SceneComponent : CachedPrePieData.TemporarilyNonTransientComponents)
+	{
+		if (SceneComponent.IsValid())
+		{
+			SceneComponent->ClearFlags(RF_NonPIEDuplicateTransient);
+		}
+	}
+
+	for (const TWeakObjectPtr<AActor>& Actor : CachedPrePieData.TemporarilyNonTransientActors)
+	{
+		if (Actor.IsValid())
+		{
+			Actor->ClearFlags(RF_NonPIEDuplicateTransient);
+		}
+	}
+
+	CachedPrePieData.TemporarilyNonTransientComponents.Empty();
+	CachedPrePieData.TemporarilyNonTransientActors.Empty();
 }
 
 ULevelSequence* UMovieGraphQuickRenderSubsystem::SetUpRenderingLevelSequence() const
diff --git a/Engine/Plugins/MovieScene/MovieRenderPipeline/Source/MovieRenderPipelineEditor/Public/Graph/MovieGraphQuickRender.h b/Engine/Plugins/MovieScene/MovieRenderPipeline/Source/MovieRenderPipelineEditor/Public/Graph/MovieGraphQuickRender.h
--- a/Engine/Plugins/MovieScene/MovieRenderPipeline/Source/MovieRenderPipelineEditor/Public/Graph/MovieGraphQuickRender.h
+++ b/Engine/Plugins/MovieScene/MovieRenderPipeline/Source/MovieRenderPipelineEditor/Public/Graph/MovieGraphQuickRender.h
@@ -59,6 +59,12 @@ private:
 		
 		/** The cameras that were selected. When PIE starts, the Outliner selection is cleared out, so it needs to be cached prior to PIE. */
 		TArray<TWeakObjectPtr<ACameraActor>> SelectedCameras;
+
+		/** Actors that have been temporarily flagged as RF_NonPIEDuplicateTransient. This enables them to be copied to the PIE world when they typically would not. */
+		TArray<TWeakObjectPtr<AActor>> TemporarilyNonTransientActors;
+
+		/** Components that have been temporarily flagged as RF_NonPIEDuplicateTransient. This enables them to be copied to the PIE world when they typically would not. */
+		TArray<TWeakObjectPtr<UActorComponent>> TemporarilyNonTransientComponents;
 	};
 	
 	/**
@@ -79,12 +85,23 @@ private:
 	/** Helper for ApplyQuickRenderUpdatesToDuplicateGraph() that applies actor editor visibility (by injecting a visibility modifier into the graph). */
 	void ApplyQuickRenderUpdatesToDuplicateGraph_ApplyEditorVisibility(const UWorld* InEditorWorld);
 
+	/** Helper for ApplyQuickRenderUpdatesToDuplicateGraph() that applies editor-only actor visibility (by injecting a visibility modifier into the graph). */
+	void ApplyQuickRenderUpdatesToDuplicateGraph_ApplyEditorOnlyActorVisibility(const UWorld* InEditorWorld);
+
 	/** Helper for ApplyQuickRenderUpdatesToDuplicateGraph() that applies viewport OCIO (by injecting overrides to file output and renderer nodes into the graph). */
 	void ApplyQuickRenderUpdatesToDuplicateGraph_ApplyOcio(const FOpenColorIODisplayConfiguration* InOcioConfiguration);
 
 	/** Helper for ApplyQuickRenderUpdatesToDuplicateGraph() that applies viewport scalability if the graph didn't specify an explicit scalability setting. */
 	void ApplyQuickRenderUpdatesToDuplicateGraph_Scalability();
 
+	/**
+	 * Adds a new collection + visibility modifier to the graph used for Quick Render. The operation name should be unique. The modifier can be
+	 * specified to either hide or show the actors in the collection, and can optionally process editor-only actors. Returns the collection query,
+	 * which can then be updated to include the actors that should be affected.
+	 */
+	template <typename QueryType>
+	QueryType* AddNewCollectionWithVisibilityModifier(const FString& InOperationName, const bool bModifierShouldHide, const bool bProcessEditorOnlyActors = false);
+
 	/** Injects a new node in the furthest-downstream position within the specified branch in the duplicated graph. */
 	UMovieGraphNode* InjectNodeIntoBranch(const TSubclassOf<UMovieGraphNode> NodeType, const FName& InBranchName) const;
 
