MoviePipeline: Added Quick Render playback preferences for the new H.264 MP4 node.

diff --git a/Engine/Plugins/MovieScene/MovieRenderPipeline/Source/MovieRenderPipelineCore/Private/Graph/MovieRenderGraphEditorSettings.cpp b/Engine/Plugins/MovieScene/MovieRenderPipeline/Source/MovieRenderPipelineCore/Private/Graph/MovieRenderGraphEditorSettings.cpp
--- a/Engine/Plugins/MovieScene/MovieRenderPipeline/Source/MovieRenderPipelineCore/Private/Graph/MovieRenderGraphEditorSettings.cpp
+++ b/Engine/Plugins/MovieScene/MovieRenderPipeline/Source/MovieRenderPipelineCore/Private/Graph/MovieRenderGraphEditorSettings.cpp
@@ -23,6 +23,10 @@ void UMovieRenderGraphEditorSettings::PostEditChangeChainProperty(FPropertyChang
 		FMovieGraphPostRenderSettings::StaticStruct()->FindPropertyByName(
 			GET_MEMBER_NAME_STRING_CHECKED(FMovieGraphPostRenderSettings, PostRenderAvidDNxHRPlayOptions));
 
+	static FProperty* MP4Options =
+		FMovieGraphPostRenderSettings::StaticStruct()->FindPropertyByName(
+			GET_MEMBER_NAME_STRING_CHECKED(FMovieGraphPostRenderSettings, PostRenderMP4PlayOptions));
+
 	auto UpdatePlayOptions = [](FMovieGraphPostRenderVideoPlayOptions* InPlayOptions)
 	{
 		if (InPlayOptions->PlaybackMethod == EMovieGraphPlaybackMethod::CustomViewer)
@@ -58,6 +62,10 @@ void UMovieRenderGraphEditorSettings::PostEditChangeChainProperty(FPropertyChang
 		{
 			UpdatePlayOptions(&PostRenderSettings.PostRenderAvidDNxHRPlayOptions);
 		}
+		else if (PropertyChangedEvent.PropertyChain.Contains(MP4Options))
+		{
+			UpdatePlayOptions(&PostRenderSettings.PostRenderMP4PlayOptions);
+		}
 	}
 	
 	Super::PostEditChangeChainProperty(PropertyChangedEvent);
diff --git a/Engine/Plugins/MovieScene/MovieRenderPipeline/Source/MovieRenderPipelineCore/Private/MoviePipelinePostRenderSettings.cpp b/Engine/Plugins/MovieScene/MovieRenderPipeline/Source/MovieRenderPipelineCore/Private/MoviePipelinePostRenderSettings.cpp
--- a/Engine/Plugins/MovieScene/MovieRenderPipeline/Source/MovieRenderPipelineCore/Private/MoviePipelinePostRenderSettings.cpp
+++ b/Engine/Plugins/MovieScene/MovieRenderPipeline/Source/MovieRenderPipelineCore/Private/MoviePipelinePostRenderSettings.cpp
@@ -18,7 +18,7 @@ FMovieGraphPostRenderImageSequencePlayOptions::FMovieGraphPostRenderImageSequenc
 }
 
 FMovieGraphPostRenderSettings::FMovieGraphPostRenderSettings()
-	: OutputTypePriorityOrder({TEXT("EXR"), TEXT("PNG"), TEXT("JPEG"), TEXT("BMP"), TEXT("MOV"), TEXT("MXF")})
+	: OutputTypePriorityOrder({TEXT("EXR"), TEXT("PNG"), TEXT("JPEG"), TEXT("MP4"), TEXT("BMP"), TEXT("MOV"), TEXT("MXF")})
 	, OutputTypePlayback(EMovieGraphOutputTypePlayback::UsePriorityOrder)
 {
 	
diff --git a/Engine/Plugins/MovieScene/MovieRenderPipeline/Source/MovieRenderPipelineCore/Public/MoviePipelinePostRenderSettings.h b/Engine/Plugins/MovieScene/MovieRenderPipeline/Source/MovieRenderPipelineCore/Public/MoviePipelinePostRenderSettings.h
--- a/Engine/Plugins/MovieScene/MovieRenderPipeline/Source/MovieRenderPipelineCore/Public/MoviePipelinePostRenderSettings.h
+++ b/Engine/Plugins/MovieScene/MovieRenderPipeline/Source/MovieRenderPipelineCore/Public/MoviePipelinePostRenderSettings.h
@@ -89,7 +89,11 @@ enum class EMovieGraphOutputTypePlayback : uint8
 	 */
 	UsePriorityOrder,
 
-	/** Play all media types that were generated from the render. */
+	/**
+	 * Play all media types that were generated from the render. Note that the highest-priority media dictates the player that will be used. In order
+	 * for lower-priority media to also be played back, the lower-priority media must be configured to use the same player that the highest-priority
+	 * is configured to use.
+	 */
 	PlayAllOutputTypes
 };
 
@@ -176,6 +180,10 @@ struct MOVIERENDERPIPELINECORE_API FMovieGraphPostRenderSettings
 	UPROPERTY(EditAnywhere, BlueprintReadWrite, Category = "Play Render Output", DisplayName = "Avid DNxHR Movies")
 	FMovieGraphPostRenderVideoPlayOptions PostRenderAvidDNxHRPlayOptions;
 
+	/** These options are used to determine how to play H.264 MP4 movies. */
+	UPROPERTY(EditAnywhere, BlueprintReadWrite, Category = "Play Render Output", DisplayName = "H.264 MP4 Movies")
+	FMovieGraphPostRenderVideoPlayOptions PostRenderMP4PlayOptions;
+
 	// TODO: No CLE support yet.
 	/** If PostRenderBehavior is set to PlayRenderOutput, these options are used to determine how to play Command Line Encoder output. */
 	// UPROPERTY(EditAnywhere, BlueprintReadWrite, Category = "Play Render Output", DisplayName = "Command Line Encoder")
diff --git a/Engine/Plugins/MovieScene/MovieRenderPipeline/Source/MovieRenderPipelineEditor/Private/MoviePipelinePostRenderFileDisplayProcessor.cpp b/Engine/Plugins/MovieScene/MovieRenderPipeline/Source/MovieRenderPipelineEditor/Private/MoviePipelinePostRenderFileDisplayProcessor.cpp
--- a/Engine/Plugins/MovieScene/MovieRenderPipeline/Source/MovieRenderPipelineEditor/Private/MoviePipelinePostRenderFileDisplayProcessor.cpp
+++ b/Engine/Plugins/MovieScene/MovieRenderPipeline/Source/MovieRenderPipelineEditor/Private/MoviePipelinePostRenderFileDisplayProcessor.cpp
@@ -130,6 +130,10 @@ FMoviePipelinePostRenderFileDisplayProcessor::FFilesToOpen FMoviePipelinePostRen
 			{
 				FilesToOpen.AppleProResMovies.Add(VideoPath);
 			}
+			else if (UpperExtension == TEXT("MP4"))
+			{
+				FilesToOpen.MP4Movies.Add(VideoPath);
+			}
 			else
 			{
 				FilesToOpen.AvidDNxHRMovies.Add(VideoPath);
@@ -202,31 +206,43 @@ void FMoviePipelinePostRenderFileDisplayProcessor::GetFilteredFilesToOpen(const
 	// If there are both images and videos returned, and Play All Output Types is specified on the images, find the media type with the highest
 	// priority and use the player for that.
 	const bool bPlayAllOutputTypes = (PostRenderSettings.OutputTypePlayback == EMovieGraphOutputTypePlayback::PlayAllOutputTypes);
-	const bool bHasImagesAndVideos = !InFilesToOpen.Images.IsEmpty() && (!InFilesToOpen.AppleProResMovies.IsEmpty() || !InFilesToOpen.AvidDNxHRMovies.IsEmpty());
+	const bool bHasVideos = !InFilesToOpen.AppleProResMovies.IsEmpty() || !InFilesToOpen.AvidDNxHRMovies.IsEmpty() || !InFilesToOpen.MP4Movies.IsEmpty();
+	const bool bHasImagesAndVideos = !InFilesToOpen.Images.IsEmpty() && bHasVideos;
 	if (bPlayAllOutputTypes && bHasImagesAndVideos)
 	{
 		const int32 ImagePriority = GetExtensionPriorityIndex(InFilesToOpen.HighestPriorityImageExtension);
 		const int32 ProResPriority = GetExtensionPriorityIndex(TEXT("MOV"));
 		const int32 AvidPriority = GetExtensionPriorityIndex(TEXT("MXF"));
+		const int32 MP4Priority = GetExtensionPriorityIndex(TEXT("MP4"));
+
+		// Find the highest priority media type
+		TArray<int32> Priorities = {ImagePriority, ProResPriority, AvidPriority, MP4Priority};
+		Priorities.Sort();
+		const int32 HighestPriority = Priorities[0];
 
 		FString PlayerWithPriority;
 
 		// Find the player that should be used, based on the media type with the highest priority
-		if ((ImagePriority < ProResPriority) && (ImagePriority < AvidPriority))
+		if (HighestPriority == ImagePriority)
 		{
 			MaybeAssignPlayOptions(&PostRenderSettings.PostRenderImageSequencePlayOptions);
 			PlayerWithPriority = PostRenderSettings.PostRenderImageSequencePlayOptions.PlayerExecutable.FilePath;
 		}
-		else if ((ProResPriority < ImagePriority) && (ProResPriority < AvidPriority))
+		else if (HighestPriority == ProResPriority)
 		{
 			MaybeAssignPlayOptions(&PostRenderSettings.PostRenderAppleProResPlayOptions);
 			PlayerWithPriority = PostRenderSettings.PostRenderAppleProResPlayOptions.PlayerExecutable.FilePath;
 		}
-		else
+		else if (HighestPriority == AvidPriority)
 		{
 			MaybeAssignPlayOptions(&PostRenderSettings.PostRenderAvidDNxHRPlayOptions);
 			PlayerWithPriority = PostRenderSettings.PostRenderAvidDNxHRPlayOptions.PlayerExecutable.FilePath;
 		}
+		else
+		{
+			MaybeAssignPlayOptions(&PostRenderSettings.PostRenderMP4PlayOptions);
+			PlayerWithPriority = PostRenderSettings.PostRenderMP4PlayOptions.PlayerExecutable.FilePath;
+		}
 
 		// Only return the media that opens with the player associated with the highest-priority media
 		if (PostRenderSettings.PostRenderImageSequencePlayOptions.PlayerExecutable.FilePath == PlayerWithPriority)
@@ -241,6 +257,10 @@ void FMoviePipelinePostRenderFileDisplayProcessor::GetFilteredFilesToOpen(const
 		{
 			OutFilteredFilesToOpen.Append(InFilesToOpen.AvidDNxHRMovies);
 		}
+		if (PostRenderSettings.PostRenderMP4PlayOptions.PlayerExecutable.FilePath == PlayerWithPriority)
+		{
+			OutFilteredFilesToOpen.Append(InFilesToOpen.MP4Movies);
+		}
 	}
 	else
 	{
@@ -256,11 +276,16 @@ void FMoviePipelinePostRenderFileDisplayProcessor::GetFilteredFilesToOpen(const
 			MaybeAssignPlayOptions(&PostRenderSettings.PostRenderAppleProResPlayOptions);
 			OutFilteredFilesToOpen.Append(InFilesToOpen.AppleProResMovies);
 		}
-		else
+		else if (!InFilesToOpen.AvidDNxHRMovies.IsEmpty())
 		{
 			MaybeAssignPlayOptions(&PostRenderSettings.PostRenderAvidDNxHRPlayOptions);
 			OutFilteredFilesToOpen.Append(InFilesToOpen.AvidDNxHRMovies);
 		}
+		else
+		{
+			MaybeAssignPlayOptions(&PostRenderSettings.PostRenderMP4PlayOptions);
+			OutFilteredFilesToOpen.Append(InFilesToOpen.MP4Movies);
+		}
 	}
 }
 
@@ -280,6 +305,7 @@ void FMoviePipelinePostRenderFileDisplayProcessor::GroupFilesByShot(
 		const bool bFirstRenderLayerOnly_Images = (PostRenderSettings.PostRenderImageSequencePlayOptions.RenderLayerPlayback == EMovieGraphRenderLayerPlaybackRange::FirstRenderLayerOnly);
 		const bool bFirstRenderLayerOnly_ProRes = (PostRenderSettings.PostRenderAppleProResPlayOptions.RenderLayerPlayback == EMovieGraphRenderLayerPlaybackRange::FirstRenderLayerOnly);
 		const bool bFirstRenderLayerOnly_Avid = (PostRenderSettings.PostRenderAvidDNxHRPlayOptions.RenderLayerPlayback == EMovieGraphRenderLayerPlaybackRange::FirstRenderLayerOnly);
+		const bool bFirstRenderLayerOnly_MP4 = (PostRenderSettings.PostRenderMP4PlayOptions.RenderLayerPlayback == EMovieGraphRenderLayerPlaybackRange::FirstRenderLayerOnly);
 		
 		// Sort the render layers by index. This will allow us to only get files from the first render layer rendered if the "First Render Layer Only"
 		// setting is turned on.
@@ -292,6 +318,7 @@ void FMoviePipelinePostRenderFileDisplayProcessor::GroupFilesByShot(
 		FString ImagesFirstFoundOnLayer;
 		FString ProResFirstFoundOnLayer;
 		FString AvidFirstFoundOnLayer;
+		FString MP4FirstFoundOnLayer;
 
 		// Add the images/videos for each render layer within the shot (only the first render layer may be considered depending on settings)
 		for (const TPair<FMovieGraphRenderDataIdentifier, FMovieGraphRenderLayerOutputData>& PassData : OutputData.RenderLayerData)
@@ -311,6 +338,7 @@ void FMoviePipelinePostRenderFileDisplayProcessor::GroupFilesByShot(
 				const bool bIsImageNode = NodeClass->IsChildOf(UMovieGraphImageSequenceOutputNode::StaticClass());
 				const bool bIsProResNode = NodeClass->GetName().Contains(TEXT("ProRes"), ESearchCase::CaseSensitive);
 				const bool bIsAvidNode = NodeClass->GetName().Contains(TEXT("Avid"), ESearchCase::CaseSensitive);
+				const bool bIsMP4Node = NodeClass->GetName().Contains(TEXT("MP4"), ESearchCase::CaseSensitive);
 				
 				// Don't continue processing this node if the node's media type has already been found for this render layer. Note that the render
 				// layer name check is to allow media of the same category (eg, several image types like PNG and EXR) to be processed for a single
@@ -327,6 +355,10 @@ void FMoviePipelinePostRenderFileDisplayProcessor::GroupFilesByShot(
 				{
 					continue;
 				}
+				if (bIsMP4Node && bFirstRenderLayerOnly_MP4 && (!MP4FirstFoundOnLayer.IsEmpty() && MP4FirstFoundOnLayer != PassData.Key.LayerName))
+				{
+					continue;
+				}
 
 				// For all files generated by the node, categorize it either by extension and/or frame template (for images)
 				for (const FString& NodeFile : NodeToFiles.Value.Array)
@@ -363,6 +395,10 @@ void FMoviePipelinePostRenderFileDisplayProcessor::GroupFilesByShot(
 						{
 							AvidFirstFoundOnLayer = PassData.Key.LayerName;
 						}
+						else if (bIsMP4Node)
+						{
+							MP4FirstFoundOnLayer = PassData.Key.LayerName;
+						}
 						
 						VideosByExtension.FindOrAdd(Extension).Add(NodeFile);
 					}
diff --git a/Engine/Plugins/MovieScene/MovieRenderPipeline/Source/MovieRenderPipelineEditor/Private/MoviePipelinePostRenderFileDisplayProcessor.h b/Engine/Plugins/MovieScene/MovieRenderPipeline/Source/MovieRenderPipelineEditor/Private/MoviePipelinePostRenderFileDisplayProcessor.h
--- a/Engine/Plugins/MovieScene/MovieRenderPipeline/Source/MovieRenderPipelineEditor/Private/MoviePipelinePostRenderFileDisplayProcessor.h
+++ b/Engine/Plugins/MovieScene/MovieRenderPipeline/Source/MovieRenderPipelineEditor/Private/MoviePipelinePostRenderFileDisplayProcessor.h
@@ -29,6 +29,7 @@ private:
 		TArray<FString> Images;
 		TArray<FString> AppleProResMovies;
 		TArray<FString> AvidDNxHRMovies;
+		TArray<FString> MP4Movies;
 
 		// The start and end frame of the first image sequence contained in Images. Representing these as strings is a little strange. It's
 		// to preserve the padding found in the source filename so it doesn't have to be re-derived later in the pipeline (where determining what that
