MoviePipeline: Updated the Quick Render settings window to include a button which opens the playback preferences for the 'Play Render Output' option.

Also moved these playback preferences to the "Plugins" category within Editor Preferences.

diff --git a/Engine/Plugins/MovieScene/MovieRenderPipeline/Source/MovieRenderPipelineCore/Private/Graph/MovieRenderGraphEditorSettings.cpp b/Engine/Plugins/MovieScene/MovieRenderPipeline/Source/MovieRenderPipelineCore/Private/Graph/MovieRenderGraphEditorSettings.cpp
--- a/Engine/Plugins/MovieScene/MovieRenderPipeline/Source/MovieRenderPipelineCore/Private/Graph/MovieRenderGraphEditorSettings.cpp
+++ b/Engine/Plugins/MovieScene/MovieRenderPipeline/Source/MovieRenderPipelineCore/Private/Graph/MovieRenderGraphEditorSettings.cpp
@@ -8,6 +8,11 @@ UMovieRenderGraphEditorSettings::UMovieRenderGraphEditorSettings()
 	
 }
 
+FName UMovieRenderGraphEditorSettings::GetCategoryName() const
+{
+	return FName(TEXT("Plugins"));
+}
+
 #if WITH_EDITOR
 void UMovieRenderGraphEditorSettings::PostEditChangeChainProperty(FPropertyChangedChainEvent& PropertyChangedEvent)
 {
diff --git a/Engine/Plugins/MovieScene/MovieRenderPipeline/Source/MovieRenderPipelineCore/Public/Graph/MovieRenderGraphEditorSettings.h b/Engine/Plugins/MovieScene/MovieRenderPipeline/Source/MovieRenderPipelineCore/Public/Graph/MovieRenderGraphEditorSettings.h
--- a/Engine/Plugins/MovieScene/MovieRenderPipeline/Source/MovieRenderPipelineCore/Public/Graph/MovieRenderGraphEditorSettings.h
+++ b/Engine/Plugins/MovieScene/MovieRenderPipeline/Source/MovieRenderPipelineCore/Public/Graph/MovieRenderGraphEditorSettings.h
@@ -18,11 +18,13 @@ class MOVIERENDERPIPELINECORE_API UMovieRenderGraphEditorSettings : public UDeve
 public:
 	UMovieRenderGraphEditorSettings();
 
+	virtual FName GetCategoryName() const override;
+
 #if WITH_EDITOR
 	virtual void PostEditChangeChainProperty(FPropertyChangedChainEvent& PropertyChangedEvent) override;
 #endif
 
 	/** If PostRenderBehavior is set to PlayRenderOutput, these settings are used to determine how to play media. */
-	UPROPERTY(Config, EditAnywhere, Category="Play Render Output", meta=(ShowOnlyInnerProperties))
+	UPROPERTY(Config, EditAnywhere, Category="Play Render Output (Quick Render only)", meta=(ShowOnlyInnerProperties))
 	FMovieGraphPostRenderSettings PostRenderSettings;
 };
diff --git a/Engine/Plugins/MovieScene/MovieRenderPipeline/Source/MovieRenderPipelineCore/Public/MoviePipelinePostRenderSettings.h b/Engine/Plugins/MovieScene/MovieRenderPipeline/Source/MovieRenderPipelineCore/Public/MoviePipelinePostRenderSettings.h
--- a/Engine/Plugins/MovieScene/MovieRenderPipeline/Source/MovieRenderPipelineCore/Public/MoviePipelinePostRenderSettings.h
+++ b/Engine/Plugins/MovieScene/MovieRenderPipeline/Source/MovieRenderPipelineCore/Public/MoviePipelinePostRenderSettings.h
@@ -155,7 +155,7 @@ struct MOVIERENDERPIPELINECORE_API FMovieGraphPostRenderSettings
 	 * When deciding which media to play, media with extensions ranked higher in this list will play instead of media ranked lower.
 	 * For example, if EXR is listed above PNG, and both EXRs and PNGs were rendered, the EXRs will play instead of the PNGs.
 	 */
-	UPROPERTY(EditAnywhere, BlueprintReadWrite, Category = "Play Render Output")
+	UPROPERTY(EditAnywhere, BlueprintReadWrite, Category = "Play Render Output (Quick Render only)")
 	TArray<FString> OutputTypePriorityOrder;
 
 	/**
@@ -165,19 +165,19 @@ struct MOVIERENDERPIPELINECORE_API FMovieGraphPostRenderSettings
 	 * player as well (ie, two or more players will not be opened). If using the Operating System to open media, this setting is effectively ignored
 	 * because only the highest-ranked media in Output Type Priority Order will be opened.
 	 */
-	UPROPERTY(EditAnywhere, BlueprintReadWrite, Category = "Play Render Output")
+	UPROPERTY(EditAnywhere, BlueprintReadWrite, Category = "Play Render Output (Quick Render only)")
 	EMovieGraphOutputTypePlayback OutputTypePlayback;
 
 	/** These options are used to determine how to play the rendered images. */
-	UPROPERTY(EditAnywhere, BlueprintReadWrite, Category = "Play Render Output", DisplayName = "Image Sequences")
+	UPROPERTY(EditAnywhere, BlueprintReadWrite, Category = "Play Render Output (Quick Render only)", DisplayName = "Image Sequences")
 	FMovieGraphPostRenderImageSequencePlayOptions PostRenderImageSequencePlayOptions;
 
 	/** These options are used to determine how to play Apple ProRes movies. */
-	UPROPERTY(EditAnywhere, BlueprintReadWrite, Category = "Play Render Output", DisplayName = "Apple ProRes Movies")
+	UPROPERTY(EditAnywhere, BlueprintReadWrite, Category = "Play Render Output (Quick Render only)", DisplayName = "Apple ProRes Movies")
 	FMovieGraphPostRenderVideoPlayOptions PostRenderAppleProResPlayOptions;
 
 	/** These options are used to determine how to play Avid DNxHR movies. */
-	UPROPERTY(EditAnywhere, BlueprintReadWrite, Category = "Play Render Output", DisplayName = "Avid DNxHR Movies")
+	UPROPERTY(EditAnywhere, BlueprintReadWrite, Category = "Play Render Output (Quick Render only)", DisplayName = "Avid DNxHR Movies")
 	FMovieGraphPostRenderVideoPlayOptions PostRenderAvidDNxHRPlayOptions;
 
 	/** These options are used to determine how to play H.264 MP4 movies. */
diff --git a/Engine/Plugins/MovieScene/MovieRenderPipeline/Source/MovieRenderPipelineEditor/Private/Customizations/Graph/MovieGraphQuickRenderSettingsCustomization.h b/Engine/Plugins/MovieScene/MovieRenderPipeline/Source/MovieRenderPipelineEditor/Private/Customizations/Graph/MovieGraphQuickRenderSettingsCustomization.h
--- a/Engine/Plugins/MovieScene/MovieRenderPipeline/Source/MovieRenderPipelineEditor/Private/Customizations/Graph/MovieGraphQuickRenderSettingsCustomization.h
+++ b/Engine/Plugins/MovieScene/MovieRenderPipeline/Source/MovieRenderPipelineEditor/Private/Customizations/Graph/MovieGraphQuickRenderSettingsCustomization.h
@@ -4,15 +4,20 @@
 
 #include "DetailCategoryBuilder.h"
 #include "DetailLayoutBuilder.h"
+#include "DetailWidgetRow.h"
 #include "Graph/MovieGraphQuickRenderSettings.h"
 #include "IDetailCustomization.h"
 #include "IDetailGroup.h"
+#include "ISettingsModule.h"
+#include "Modules/ModuleManager.h"
 #include "MovieGraphCustomizationUtils.h"
 #include "MoviePipelineQueue.h"
 #include "MoviePipelineQuickRenderUIState.h"
 #include "MoviePipelineUtils.h"
 #include "PropertyHandle.h"
 #include "UObject/Package.h"
+#include "Widgets/Images/SImage.h"
+#include "Widgets/Input/SButton.h"
 
 #define LOCTEXT_NAMESPACE "MoviePipelineQuickRenderSettings"
 
@@ -90,6 +95,57 @@ protected:
 		DetailBuilder->EditDefaultProperty(CustomEndFrameProperty)->IsEnabled(IsEnabledAttr);
 	}
 
+	// Adds a button next to the post-render behavior drop-down to open the editor preferences that dictate how media is played back.
+	void AddPostRenderPlayOptionsButton() const
+	{
+		const TSharedRef<IPropertyHandle> PostRenderBehaviorProperty =
+			DetailBuilder->GetProperty(GET_MEMBER_NAME_CHECKED(UMovieGraphQuickRenderModeSettings, PostRenderBehavior), UMovieGraphQuickRenderModeSettings::StaticClass());
+		IDetailPropertyRow* PostRenderBehaviorRow = DetailBuilder->EditDefaultProperty(PostRenderBehaviorProperty);
+
+		TSharedPtr<SWidget> NameWidget;
+		TSharedPtr<SWidget> ValueWidget;
+		PostRenderBehaviorRow->GetDefaultWidgets(NameWidget, ValueWidget);
+
+		PostRenderBehaviorRow->CustomWidget()
+		.NameContent()
+		[
+			NameWidget.ToSharedRef()
+		]
+		.ValueContent()
+		[
+			SNew(SHorizontalBox)
+			+ SHorizontalBox::Slot()
+			.AutoWidth()
+			[
+				ValueWidget.ToSharedRef()
+			]
+
+			+ SHorizontalBox::Slot()
+			.FillWidth(1.f)
+			.HAlign(HAlign_Left)
+			.VAlign(VAlign_Center)
+			.Padding(5.f, 0, 0, 0)
+			[
+				SNew(SButton)
+				.ButtonStyle(FAppStyle::Get(), TEXT("SimpleButton"))
+				.ContentPadding(0)
+				.OnClicked_Lambda([]()
+				{
+					FModuleManager::LoadModuleChecked<ISettingsModule>("Settings").ShowViewer("Editor", "Plugins", "MovieRenderGraphEditorSettings");
+					
+					return FReply::Handled();
+				})
+				.ToolTipText(LOCTEXT("OpenPlaybackPrefs_Tooltip", "Open up the playback preferences which apply to the 'Play Render Output' option."))
+				[
+					SNew(SImage)
+					.ColorAndOpacity(FSlateColor::UseForeground())
+					.Image(FAppStyle::Get().GetBrush("EditorPreferences.TabIcon"))
+					.DesiredSizeOverride(FVector2D(16.f, 16.f))
+				]
+			]
+		];
+	}
+
 	virtual void CustomizeDetails(IDetailLayoutBuilder& InDetailBuilder) override
 	{
 		// Refresh the customization every time a save happens. Use this opportunity to update the variables in the UI. We could update the UI before
@@ -138,6 +194,8 @@ protected:
 		ConfigurationCategory.SetSortOrder(SortOrder);
 		QuickRenderCategory.SetSortOrder(++SortOrder);
 		PrimaryGraphVariablesCategory.SetSortOrder(++SortOrder);
+
+		AddPostRenderPlayOptionsButton();
 	}
 	//~ End IDetailCustomization interface
 
