MoviePipeline: Added telemetry for Quick Render.

diff --git a/Engine/Plugins/MovieScene/MovieRenderPipeline/Source/MovieRenderPipelineCore/Private/MoviePipelineTelemetry.cpp b/Engine/Plugins/MovieScene/MovieRenderPipeline/Source/MovieRenderPipelineCore/Private/MoviePipelineTelemetry.cpp
--- a/Engine/Plugins/MovieScene/MovieRenderPipeline/Source/MovieRenderPipelineCore/Private/MoviePipelineTelemetry.cpp
+++ b/Engine/Plugins/MovieScene/MovieRenderPipeline/Source/MovieRenderPipelineCore/Private/MoviePipelineTelemetry.cpp
@@ -5,6 +5,7 @@
 #include "EngineAnalytics.h"
 #include "Graph/MovieGraphBlueprintLibrary.h"
 #include "Graph/MovieGraphConfig.h"
+#include "Graph/MovieGraphQuickRenderSettings.h"
 #include "MoviePipelineBlueprintLibrary.h"
 #include "MoviePipelineQueue.h"
 
@@ -14,6 +15,8 @@
 
 void FMoviePipelineTelemetry::SendRendersRequestedTelemetry(const bool bIsLocal, TArray<UMoviePipelineExecutorJob*>&& InJobs)
 {
+	bIsCurrentRenderRequestQuickRender = false;
+	
 	if (!FEngineAnalytics::IsAvailable())
 	{
 		return;
@@ -59,6 +62,30 @@ void FMoviePipelineTelemetry::SendRendersRequestedTelemetry(const bool bIsLocal,
 	FEngineAnalytics::GetProvider().RecordEvent(TEXT("MoviePipeline.RendersRequested"), EventAttributes);
 }
 
+void FMoviePipelineTelemetry::SendQuickRenderRequestedTelemetry(const EMovieGraphQuickRenderMode QuickRenderMode)
+{
+	bIsCurrentRenderRequestQuickRender = true;
+	
+	if (!FEngineAnalytics::IsAvailable())
+	{
+		return;
+	}
+
+	UMovieGraphQuickRenderModeSettings* ModeSettings = UMovieGraphQuickRenderSettings::GetSavedQuickRenderModeSettings(QuickRenderMode);
+
+	TArray<FAnalyticsEventAttribute> EventAttributes;
+	EventAttributes.Add(FAnalyticsEventAttribute(TEXT("Mode"),
+		StaticEnum<EMovieGraphQuickRenderMode>()->GetNameStringByValue(static_cast<int64>(QuickRenderMode))));
+	EventAttributes.Add(FAnalyticsEventAttribute(TEXT("PlayAfterRenderMode"),
+		StaticEnum<EMovieGraphQuickRenderPostRenderActionType>()->GetNameStringByValue(static_cast<int64>(ModeSettings->PostRenderBehavior))));
+	EventAttributes.Add(FAnalyticsEventAttribute(TEXT("ViewportLookFlags"),
+		StaticEnum<EMovieGraphQuickRenderViewportLookFlags>()->GetValueOrBitfieldAsString(ModeSettings->ViewportLookFlags)));
+	EventAttributes.Add(FAnalyticsEventAttribute(TEXT("SequencerFrameRange"),
+		StaticEnum<EMovieGraphQuickRenderFrameRangeType>()->GetNameStringByValue(static_cast<int64>(ModeSettings->FrameRangeType))));
+
+	FEngineAnalytics::GetProvider().RecordEvent(TEXT("MoviePipeline.QuickRenderRequested"), EventAttributes);
+}
+
 /** Sends out telemetry that includes information about the shot being rendered (the type of rendering being done, which types of settings/nodes are being used, etc). */
 void FMoviePipelineTelemetry::SendBeginShotRenderTelemetry(UMoviePipelineExecutorShot* InShot, UMovieGraphEvaluatedConfig* InEvaluatedConfig)
 {
@@ -73,6 +100,7 @@ void FMoviePipelineTelemetry::SendBeginShotRenderTelemetry(UMoviePipelineExecuto
 	
 	TArray<FAnalyticsEventAttribute> EventAttributes;
 	EventAttributes.Add(FAnalyticsEventAttribute(TEXT("IsGraph"), ShotRenderTelemetry.bIsGraph));
+	EventAttributes.Add(FAnalyticsEventAttribute(TEXT("IsQuickRender"), ShotRenderTelemetry.bIsQuickRender));
 	EventAttributes.Add(FAnalyticsEventAttribute(TEXT("UsesDeferred"), ShotRenderTelemetry.bUsesDeferred));
 	EventAttributes.Add(FAnalyticsEventAttribute(TEXT("UsesPathTracer"), ShotRenderTelemetry.bUsesPathTracer));
 	EventAttributes.Add(FAnalyticsEventAttribute(TEXT("UsesPanoramic"), ShotRenderTelemetry.bUsesPanoramic));
@@ -107,6 +135,7 @@ void FMoviePipelineTelemetry::SendEndShotRenderTelemetry(const bool bIsGraph, co
 
 	TArray<FAnalyticsEventAttribute> EventAttributes;
 	EventAttributes.Add(FAnalyticsEventAttribute(TEXT("IsGraph"), bIsGraph));
+	EventAttributes.Add(FAnalyticsEventAttribute(TEXT("IsQuickRender"), bIsCurrentRenderRequestQuickRender));
 	EventAttributes.Add(FAnalyticsEventAttribute(TEXT("WasSuccessful"), bWasSuccessful));
 	EventAttributes.Add(FAnalyticsEventAttribute(TEXT("WasCanceled"), bWasCanceled));
 
@@ -137,6 +166,7 @@ FMoviePipelineShotRenderTelemetry FMoviePipelineTelemetry::GatherShotRenderTelem
 {
 	FMoviePipelineShotRenderTelemetry Telemetry;
 	Telemetry.bIsGraph = true;
+	Telemetry.bIsQuickRender = bIsCurrentRenderRequestQuickRender;
 
 	const FIntPoint Resolution = UMovieGraphBlueprintLibrary::GetOverscannedResolution(InEvaluatedConfig);
 	Telemetry.ResolutionX = Resolution.X;
diff --git a/Engine/Plugins/MovieScene/MovieRenderPipeline/Source/MovieRenderPipelineCore/Public/Graph/MovieGraphQuickRenderSettings.h b/Engine/Plugins/MovieScene/MovieRenderPipeline/Source/MovieRenderPipelineCore/Public/Graph/MovieGraphQuickRenderSettings.h
--- a/Engine/Plugins/MovieScene/MovieRenderPipeline/Source/MovieRenderPipelineCore/Public/Graph/MovieGraphQuickRenderSettings.h
+++ b/Engine/Plugins/MovieScene/MovieRenderPipeline/Source/MovieRenderPipelineCore/Public/Graph/MovieGraphQuickRenderSettings.h
@@ -69,7 +69,7 @@ enum class EMovieGraphQuickRenderPostRenderActionType : uint8
 };
 
 /** The aspects of the viewport that should be applied to the quick render. */
-UENUM(BlueprintType, meta = (Bitflags, UseEnumValuesAsMaskValuesInEditor = "true"))
+UENUM(BlueprintType, Flags, meta = (Bitflags, UseEnumValuesAsMaskValuesInEditor = "true"))
 enum class EMovieGraphQuickRenderViewportLookFlags : uint8
 {
 	None		= 0	UMETA(Hidden),
diff --git a/Engine/Plugins/MovieScene/MovieRenderPipeline/Source/MovieRenderPipelineCore/Public/MoviePipelineTelemetry.h b/Engine/Plugins/MovieScene/MovieRenderPipeline/Source/MovieRenderPipelineCore/Public/MoviePipelineTelemetry.h
--- a/Engine/Plugins/MovieScene/MovieRenderPipeline/Source/MovieRenderPipelineCore/Public/MoviePipelineTelemetry.h
+++ b/Engine/Plugins/MovieScene/MovieRenderPipeline/Source/MovieRenderPipelineCore/Public/MoviePipelineTelemetry.h
@@ -3,6 +3,7 @@
 #pragma once
 
 #include "Containers/Array.h"
+#include "Graph/MovieGraphQuickRenderSettings.h"
 #include "HAL/Platform.h"
 
 class UMovieGraphEvaluatedConfig;
@@ -13,6 +14,7 @@ class UMoviePipelineExecutorShot;
 struct MOVIERENDERPIPELINECORE_API FMoviePipelineShotRenderTelemetry
 {
 	bool bIsGraph = false;
+	bool bIsQuickRender = false;
 	bool bUsesDeferred = false;
 	bool bUsesPathTracer = false;
 	bool bUsesPanoramic = false;
@@ -38,13 +40,16 @@ struct MOVIERENDERPIPELINECORE_API FMoviePipelineShotRenderTelemetry
 	// Also remember to track the telemetry in both the graph and legacy.
 };
 
-/** Responsible for sending out telemetry for both Movie Render Queue and Movie Render Graph. */
+/** Responsible for sending out telemetry for Movie Render Queue, Movie Render Graph, and Quick Render. */
 class MOVIERENDERPIPELINECORE_API FMoviePipelineTelemetry
 {
 public:
-	/** Sends out telemetry that captures queue-level information. */
+	/** Sends out telemetry that captures queue-level information. Called by either MRQ or MRG, not Quick Render. */
 	static void SendRendersRequestedTelemetry(const bool bIsLocal, TArray<UMoviePipelineExecutorJob*>&& InJobs);
 
+	/** Sends out telemetry that captures Quick Render information for the provided mode. */
+	static void SendQuickRenderRequestedTelemetry(const EMovieGraphQuickRenderMode QuickRenderMode);
+
 	/** Sends out telemetry that includes information about the shot being rendered (the type of rendering being done, which types of settings/nodes are being used, etc). */
 	static void SendBeginShotRenderTelemetry(UMoviePipelineExecutorShot* InShot, UMovieGraphEvaluatedConfig* InEvaluatedConfig = nullptr);
 
@@ -52,6 +57,9 @@ public:
 	static void SendEndShotRenderTelemetry(const bool bIsGraph, const bool bWasSuccessful, const bool bWasCanceled);
 
 private:
+	/** Tracks whether the current render request originated from Quick Render. If false, the render originated from a queue in either MRQ or MRG. */
+	inline static bool bIsCurrentRenderRequestQuickRender = false;
+	
 	/** Gets the current session type (Editor, DashGame, Shipping). */
 	static FString GetSessionType();
 
diff --git a/Engine/Plugins/MovieScene/MovieRenderPipeline/Source/MovieRenderPipelineEditor/Private/Graph/MovieGraphQuickRender.cpp b/Engine/Plugins/MovieScene/MovieRenderPipeline/Source/MovieRenderPipelineEditor/Private/Graph/MovieGraphQuickRender.cpp
--- a/Engine/Plugins/MovieScene/MovieRenderPipeline/Source/MovieRenderPipelineEditor/Private/Graph/MovieGraphQuickRender.cpp
+++ b/Engine/Plugins/MovieScene/MovieRenderPipeline/Source/MovieRenderPipelineEditor/Private/Graph/MovieGraphQuickRender.cpp
@@ -16,6 +16,7 @@
 #include "MoviePipelinePIEExecutor.h"
 #include "MoviePipelinePostRenderFileDisplayProcessor.h"
 #include "MoviePipelineQueueSubsystem.h"
+#include "MoviePipelineTelemetry.h"
 #include "MoviePipelineUtils.h"
 #include "MovieRenderPipelineCoreModule.h"
 #include "MovieRenderPipelineSettings.h"
@@ -171,6 +172,8 @@ void UMovieGraphQuickRenderSubsystem::BeginQuickRender(const EMovieGraphQuickRen
 		PerformPostRenderCleanup();
 	});
 
+	FMoviePipelineTelemetry::SendQuickRenderRequestedTelemetry(InQuickRenderMode);
+
 	// Do the render
 	Subsystem->RenderQueueInstanceWithExecutorInstance(TemporaryQueue.Get(), TemporaryExecutor.Get());
 }
