MoviePipeline: Fixed an issue where the MP4 module was not being loaded.

diff --git a/Engine/Plugins/MovieScene/MovieRenderPipeline/MovieRenderPipeline.uplugin b/Engine/Plugins/MovieScene/MovieRenderPipeline/MovieRenderPipeline.uplugin
--- a/Engine/Plugins/MovieScene/MovieRenderPipeline/MovieRenderPipeline.uplugin
+++ b/Engine/Plugins/MovieScene/MovieRenderPipeline/MovieRenderPipeline.uplugin
@@ -72,6 +72,12 @@
 			"Type": "Editor",
 			"LoadingPhase": "Default"
 		},
+		{
+			"Name": "MovieRenderPipelineMP4Encoder",
+			"Type": "Runtime",
+			"LoadingPhase": "Default",
+			"PlatformAllowList": ["Mac", "Win64", "Linux" ]
+		},
 		{
 			"Name": "UEOpenExrRTTI",
 			"Type": "Runtime",
diff --git a/Engine/Plugins/MovieScene/MovieRenderPipeline/Source/MovieRenderPipelineMP4Encoder/MovieRenderPipelineMP4Encoder.Build.cs b/Engine/Plugins/MovieScene/MovieRenderPipeline/Source/MovieRenderPipelineMP4Encoder/MovieRenderPipelineMP4Encoder.Build.cs
--- a/Engine/Plugins/MovieScene/MovieRenderPipeline/Source/MovieRenderPipelineMP4Encoder/MovieRenderPipelineMP4Encoder.Build.cs
+++ b/Engine/Plugins/MovieScene/MovieRenderPipeline/Source/MovieRenderPipelineMP4Encoder/MovieRenderPipelineMP4Encoder.Build.cs
@@ -35,6 +35,14 @@ namespace UnrealBuildTool.Rules
 
 				PrivateIncludePaths.Add("MovieRenderPipelineMP4Encoder/Private/Windows/");
 			}
+			else if(Target.Platform.IsInGroup(UnrealPlatformGroup.Linux))
+			{
+				PrivateIncludePaths.Add("MovieRenderPipelineMP4Encoder/Private/GenericPlatform/");
+			}
+			else if(Target.Platform.IsInGroup(UnrealPlatformGroup.Apple))
+			{
+				PrivateIncludePaths.Add("MovieRenderPipelineMP4Encoder/Private/GenericPlatform/");
+			}
         }
     }
 }
diff --git a/Engine/Plugins/MovieScene/MovieRenderPipeline/Source/MovieRenderPipelineMP4Encoder/Private/GenericPlatform/MoviePipelineMP4Encoder.h b/Engine/Plugins/MovieScene/MovieRenderPipeline/Source/MovieRenderPipelineMP4Encoder/Private/GenericPlatform/MoviePipelineMP4Encoder.h
--- a/Engine/Plugins/MovieScene/MovieRenderPipeline/Source/MovieRenderPipelineMP4Encoder/Private/GenericPlatform/MoviePipelineMP4Encoder.h
+++ b/Engine/Plugins/MovieScene/MovieRenderPipeline/Source/MovieRenderPipelineMP4Encoder/Private/GenericPlatform/MoviePipelineMP4Encoder.h
@@ -9,7 +9,7 @@
 
 /**
 * Dummy implementation for unsupported platforms which returns failure if used, to avoid
-* doing rendering work and then having no output. This should have a matching pubilc API
+* doing rendering work and then having no output. This should have a matching public API
 * to the implementation in the other platform folders.
 */
 class FMoviePipelineMP4Encoder
