Movie Render Queue:  Auto-exposure pass opt-in flags removed from UX for high resolution and panoramic rendering.  Auto-exposure pass is automatically enabled, as long as manual exposure is not enabled and per-tile view state is enabled.  System memory mirroring added to panoramic rendering, with a caveat that TSR buffers currently aren't mirrored, and can be a memory bottleneck if enabled (TSR is force disabled for high resolution).  Fixed missing constructor initialization of bSystemMemoryMirror in high res settings, and updated some comments.

#rnx

diff --git a/Engine/Plugins/MovieScene/MoviePipelineMaskRenderPass/Source/MoviePipelineMaskRenderPass/MoviePipelineMaskRenderPass.Build.cs b/Engine/Plugins/MovieScene/MoviePipelineMaskRenderPass/Source/MoviePipelineMaskRenderPass/MoviePipelineMaskRenderPass.Build.cs
--- a/Engine/Plugins/MovieScene/MoviePipelineMaskRenderPass/Source/MoviePipelineMaskRenderPass/MoviePipelineMaskRenderPass.Build.cs
+++ b/Engine/Plugins/MovieScene/MoviePipelineMaskRenderPass/Source/MoviePipelineMaskRenderPass/MoviePipelineMaskRenderPass.Build.cs
@@ -13,6 +13,7 @@ public class MoviePipelineMaskRenderPass : ModuleRules
 				"Engine",
 				"Json",
 				"ImageCore",
+				"MovieScene",
 			}
 		);
 
diff --git a/Engine/Plugins/MovieScene/MoviePipelineMaskRenderPass/Source/MoviePipelineMaskRenderPass/Private/MoviePipelinePanoramicPass.cpp b/Engine/Plugins/MovieScene/MoviePipelineMaskRenderPass/Source/MoviePipelineMaskRenderPass/Private/MoviePipelinePanoramicPass.cpp
--- a/Engine/Plugins/MovieScene/MoviePipelineMaskRenderPass/Source/MoviePipelineMaskRenderPass/Private/MoviePipelinePanoramicPass.cpp
+++ b/Engine/Plugins/MovieScene/MoviePipelineMaskRenderPass/Source/MoviePipelineMaskRenderPass/Private/MoviePipelinePanoramicPass.cpp
@@ -16,6 +16,7 @@
 #include "OpenColorIODisplayExtension.h"
 #include "TextureResource.h"
 #include "SceneUtils.h"
+#include "Camera/CameraComponent.h"
 
 #include UE_INLINE_GENERATED_CPP_BY_NAME(MoviePipelinePanoramicPass)
 
@@ -28,7 +29,7 @@ UMoviePipelinePanoramicPass::UMoviePipelinePanoramicPass()
 	, EyeSeparation(6.5f)
 	, EyeConvergenceDistance(EyeSeparation * 30.f)
 	, bAllocateHistoryPerPane(false)
-	, bAutoExposureRenderPass(true)
+	, bSystemMemoryMirror(false)
 	, bHasWarnedSettings(false)
 {
 	PassIdentifier = FMoviePipelinePassIdentifier("Panoramic");
@@ -128,6 +129,11 @@ void UMoviePipelinePanoramicPass::SetupImpl(const MoviePipeline::FMoviePipelineR
 	int32 NumPanoramicPanes = NumPanes * StereoMultiplier;
 	if (bAllocateHistoryPerPane)
 	{
+		if (bSystemMemoryMirror)
+		{
+			SystemMemoryMirror = FSceneViewStateInterface::SystemMemoryMirrorAllocate();
+		}
+
 		OptionalPaneViewStates.SetNum(NumPanoramicPanes);
 
 		for (int32 Index = 0; Index < OptionalPaneViewStates.Num(); Index++)
@@ -135,7 +141,20 @@ void UMoviePipelinePanoramicPass::SetupImpl(const MoviePipeline::FMoviePipelineR
 			OptionalPaneViewStates[Index].Allocate(InPassInitSettings.FeatureLevel);
 		}
 
-		if (bAutoExposureRenderPass && ComputeAutoExposureCubeCaptureSize(PaneResolution) > 0)
+		bool bManualExposureEnabled = false;
+
+		APlayerCameraManager* PlayerCameraManager = GetPipeline()->GetWorld()->GetFirstPlayerController()->PlayerCameraManager;
+		if (PlayerCameraManager->ViewTarget.Target)
+		{
+			UCameraComponent* CameraComponent = MovieSceneHelpers::CameraComponentFromActor(PlayerCameraManager->ViewTarget.Target);
+			if (CameraComponent && CameraComponent->PostProcessBlendWeight > 0.0f && CameraComponent->PostProcessSettings.bOverride_AutoExposureMethod)
+			{
+				bManualExposureEnabled = CameraComponent->PostProcessSettings.AutoExposureMethod == AEM_Manual;
+			}
+		}
+
+		// Set up view states for the auto-exposure cube map pass, if manual exposure isn't enabled
+		if (!bManualExposureEnabled && ComputeAutoExposureCubeCaptureSize(PaneResolution) > 0)
 		{
 			AutoExposureViewStates.SetNum(CubeFace_MAX);
 
@@ -579,7 +598,6 @@ void UMoviePipelinePanoramicPass::RenderSample_GameThreadImpl(const FMoviePipeli
 			for (int32 EyeLoopIndex = 0; EyeLoopIndex < NumEyeRenders; EyeLoopIndex++)
 			{
 				FMoviePipelineRenderPassMetrics InOutSampleState = InSampleState;
-				InOutSampleState.bAutoExposureRenderPass = AutoExposureViewStates.Num() > 0;
 
 				// Construct a pane that will be attached to the render information
 				FPanoPane Pane;
@@ -609,6 +627,14 @@ void UMoviePipelinePanoramicPass::RenderSample_GameThreadImpl(const FMoviePipeli
 				// the global output resolution, which is the final image size.
 				Pane.Resolution = PaneResolution;
 
+				// Optional system memory mirroring
+				if (bAllocateHistoryPerPane && bSystemMemoryMirror)
+				{
+					FSceneViewStateInterface* SceneViewState = GetSceneViewStateInterface(&Pane);
+					check(SceneViewState);
+					SceneViewState->SystemMemoryMirrorRestore(SystemMemoryMirror.Get());
+				}
+
 				// Create the View Family for this render. This will only contain one view to better fit into the existing
 				// architecture of MRQ we have. Calculating the view family requires calculating the FSceneView itself which
 				// is highly customized for panos. So we provide the FPanoPlane as 'raw' data that gets passed along so we can
@@ -648,6 +674,14 @@ void UMoviePipelinePanoramicPass::RenderSample_GameThreadImpl(const FMoviePipeli
 
 				// Schedule a readback and then high-res accumulation.
 				ScheduleReadbackAndAccumulation(InOutSampleState, Pane, Canvas);
+
+				// Optional system memory mirroring
+				if (bAllocateHistoryPerPane && bSystemMemoryMirror)
+				{
+					FSceneViewStateInterface* SceneViewState = GetSceneViewStateInterface(&Pane);
+					check(SceneViewState);
+					SceneViewState->SystemMemoryMirrorBackup(SystemMemoryMirror.Get());
+				}
 			}
 		}
 	}
diff --git a/Engine/Plugins/MovieScene/MoviePipelineMaskRenderPass/Source/MoviePipelineMaskRenderPass/Private/MoviePipelinePanoramicPass.h b/Engine/Plugins/MovieScene/MoviePipelineMaskRenderPass/Source/MoviePipelineMaskRenderPass/Private/MoviePipelinePanoramicPass.h
--- a/Engine/Plugins/MovieScene/MoviePipelineMaskRenderPass/Source/MoviePipelineMaskRenderPass/Private/MoviePipelinePanoramicPass.h
+++ b/Engine/Plugins/MovieScene/MoviePipelineMaskRenderPass/Source/MoviePipelineMaskRenderPass/Private/MoviePipelinePanoramicPass.h
@@ -2,6 +2,7 @@
 #pragma once
 
 #include "MoviePipelineImagePassBase.h"
+#include "Templates/PimplPtr.h"
 #include "MoviePipelinePanoramicPass.generated.h"
 
 class FCanvas;
@@ -12,6 +13,7 @@ struct FImageOverlappedAccumulator;
 class FSceneViewFamily;
 class FSceneView;
 struct FAccumulatorPool;
+struct FSceneViewStateSystemMemoryMirror;
 
 struct FPanoPane : public UMoviePipelineImagePassBase::IViewCalcPayload
 {
@@ -192,17 +194,16 @@ public:
 	
 	/**
 	* Should we store the render scene history per individual render? This can consume a great deal of memory with many renders,
-	* but enables TAA and other history-based effects (denoisers, etc.) to work.
+	* but enables TAA and other history-based effects (denoisers, auto-exposure, Lumen, etc.) to work.
 	*/
 	UPROPERTY(EditAnywhere, BlueprintReadWrite, AdvancedDisplay, Category = "Panoramic Settings")
 	bool bAllocateHistoryPerPane;
 
 	/*
-	* If true, auto-exposure is generated from a separate low resolution cube render pass, allowing it to be used with panoramic rendering.
-	* Requires "Allocate History Per Pane".
+	* If true, persistent GPU data per panoramic pane is mirrored to system memory, allowing higher resolutions, but increasing render time.
 	*/
-	UPROPERTY(EditAnywhere, BlueprintReadWrite, AdvancedDisplay, Category = "Panoramic Settings", meta = (editcondition = "bAllocateHistoryPerPane"))
-	bool bAutoExposureRenderPass;
+	UPROPERTY(EditAnywhere, BlueprintReadWrite, Category = "Render Settings")
+	bool bSystemMemoryMirror;
 
 protected:
 	TSharedPtr<FAccumulatorPool, ESPMode::ThreadSafe> AccumulatorPool;
@@ -218,4 +219,6 @@ protected:
 
 	TSharedPtr<MoviePipeline::IMoviePipelineOutputMerger> PanoramicOutputBlender;
 	bool bHasWarnedSettings;
+
+	TPimplPtr<FSceneViewStateSystemMemoryMirror> SystemMemoryMirror;
 };
diff --git a/Engine/Plugins/MovieScene/MovieRenderPipeline/Source/MovieRenderPipelineCore/Private/MoviePipelineRendering.cpp b/Engine/Plugins/MovieScene/MovieRenderPipeline/Source/MovieRenderPipelineCore/Private/MoviePipelineRendering.cpp
--- a/Engine/Plugins/MovieScene/MovieRenderPipeline/Source/MovieRenderPipelineCore/Private/MoviePipelineRendering.cpp
+++ b/Engine/Plugins/MovieScene/MovieRenderPipeline/Source/MovieRenderPipelineCore/Private/MoviePipelineRendering.cpp
@@ -443,7 +443,6 @@ void UMoviePipeline::RenderFrame()
 				SampleState.TextureSharpnessBias = HighResSettings->TextureSharpnessBias;
 				SampleState.OCIOConfiguration = ColorSettings ? &ColorSettings->OCIOConfiguration : nullptr;
 				SampleState.GlobalScreenPercentageFraction = FLegacyScreenPercentageDriver::GetCVarResolutionFraction();
-				SampleState.bAutoExposureRenderPass = HighResSettings->bAutoExposureRenderPass && HighResSettings->TileCount > 1;
 				SampleState.bAutoExposureCubePass = false;
 				SampleState.AutoExposureCubeFace = 0;
 				SampleState.bOverrideCameraOverscan = CameraSettings->bOverrideCameraOverscan;
diff --git a/Engine/Plugins/MovieScene/MovieRenderPipeline/Source/MovieRenderPipelineCore/Public/MoviePipelineHighResSetting.h b/Engine/Plugins/MovieScene/MovieRenderPipeline/Source/MovieRenderPipelineCore/Public/MoviePipelineHighResSetting.h
--- a/Engine/Plugins/MovieScene/MovieRenderPipeline/Source/MovieRenderPipelineCore/Public/MoviePipelineHighResSetting.h
+++ b/Engine/Plugins/MovieScene/MovieRenderPipeline/Source/MovieRenderPipelineCore/Public/MoviePipelineHighResSetting.h
@@ -23,7 +23,7 @@ public:
 		, bOverrideSubSurfaceScattering(false)
 		, BurleySampleCount(64)
 		, bAllocateHistoryPerTile(true)
-		, bAutoExposureRenderPass(true)
+		, bSystemMemoryMirror(false)
 	{
 	}
 	
@@ -52,10 +52,6 @@ public:
 			{
 				ValidationResults.Add(NSLOCTEXT("MovieRenderPipeline", "HighRes_OverlapNeeded", "Increase the overlap amount to avoid seams in the final image. More overlap results in longer renders so start at 0.1 and increase as needed."));
 			}
-			if (!bAutoExposureRenderPass)
-			{
-				ValidationResults.Add(NSLOCTEXT("MovieRenderPipeline", "HighRes_ManualExposureNeeded", "Use Manual Camera Exposure /w Exposure Compensation when the Auto Exposure Render Pass is disabled."));
-			}
 			ValidationState = EMoviePipelineValidationState::Warnings;
 		}
 	}
@@ -206,20 +202,14 @@ public:
 	int32 BurleySampleCount;
 	
 	/*
-	* If true, allocate a unique history for each tile. This is needed to make some render features work, but should be disabled
-	* when dealing with extremely large resolutions as you will spend all of your GPU memory on history buffers.
+	* If true, allocate a unique history for each tile. This is needed to make some render features work, including auto-exposure and Lumen,
+	* but should be disabled when dealing with extremely large resolutions as you will spend all of your GPU memory on history buffers.
 	*/
 	UPROPERTY(EditAnywhere, BlueprintReadWrite, Category = "Render Settings")
 	bool bAllocateHistoryPerTile;
 
 	/*
-	* If true, auto-exposure is generated from a separate low resolution full screen pass, allowing it to be used with tiling.
-	*/
-	UPROPERTY(EditAnywhere, BlueprintReadWrite, Category = "Render Settings")
-	bool bAutoExposureRenderPass;
-
-	/*
-	* If true, persistent data per high resolution tile is mirrored to system memory, allowing higher resolutions, but increasing render time.
+	* If true, persistent GPU data per high resolution tile is mirrored to system memory, allowing higher resolutions, but increasing render time.
 	*/
 	UPROPERTY(EditAnywhere, BlueprintReadWrite, Category = "Render Settings")
 	bool bSystemMemoryMirror;
diff --git a/Engine/Plugins/MovieScene/MovieRenderPipeline/Source/MovieRenderPipelineCore/Public/MovieRenderPipelineDataTypes.h b/Engine/Plugins/MovieScene/MovieRenderPipeline/Source/MovieRenderPipelineCore/Public/MovieRenderPipelineDataTypes.h
--- a/Engine/Plugins/MovieScene/MovieRenderPipeline/Source/MovieRenderPipelineCore/Public/MovieRenderPipelineDataTypes.h
+++ b/Engine/Plugins/MovieScene/MovieRenderPipeline/Source/MovieRenderPipelineCore/Public/MovieRenderPipelineDataTypes.h
@@ -1062,12 +1062,6 @@ public:
 	*/
 	float GlobalScreenPercentageFraction;
 
-	/**
-	* Whether the auto exposure render pass is enabled in the settings.  Needs to be set on ALL passes, so the other passes know to look for the results from the exposure pass.
-	* The exposure rendering itself runs with the first tile.
-	*/
-	bool bAutoExposureRenderPass;
-
 	/**
 	* Whether the auto exposure pass is an omnidirectional cube capture pass, useful for panoramic rendering.  Only set for the actual exposure pass!
 	*/
diff --git a/Engine/Plugins/MovieScene/MovieRenderPipeline/Source/MovieRenderPipelineRenderPasses/MovieRenderPipelineRenderPasses.Build.cs b/Engine/Plugins/MovieScene/MovieRenderPipeline/Source/MovieRenderPipelineRenderPasses/MovieRenderPipelineRenderPasses.Build.cs
--- a/Engine/Plugins/MovieScene/MovieRenderPipeline/Source/MovieRenderPipelineRenderPasses/MovieRenderPipelineRenderPasses.Build.cs
+++ b/Engine/Plugins/MovieScene/MovieRenderPipeline/Source/MovieRenderPipelineRenderPasses/MovieRenderPipelineRenderPasses.Build.cs
@@ -23,6 +23,7 @@ public class MovieRenderPipelineRenderPasses : ModuleRules
 				"ImageWrapper",				
 				"CinematicCamera", // For metadata
 				"MovieRenderPipelineSettings", // For settings
+				"MovieScene",
 				"SlateCore",
 				"ImageCore",
 				"OpenColorIOWrapper",
diff --git a/Engine/Plugins/MovieScene/MovieRenderPipeline/Source/MovieRenderPipelineRenderPasses/Private/MoviePipelineDeferredPasses.cpp b/Engine/Plugins/MovieScene/MovieRenderPipeline/Source/MovieRenderPipelineRenderPasses/Private/MoviePipelineDeferredPasses.cpp
--- a/Engine/Plugins/MovieScene/MovieRenderPipeline/Source/MovieRenderPipelineRenderPasses/Private/MoviePipelineDeferredPasses.cpp
+++ b/Engine/Plugins/MovieScene/MovieRenderPipeline/Source/MovieRenderPipelineRenderPasses/Private/MoviePipelineDeferredPasses.cpp
@@ -123,6 +123,41 @@ FIntPoint UMoviePipelineDeferredPassBase::GetOverscannedResolutionForCamera(cons
 	return UMoviePipelineBlueprintLibrary::GetOverscannedResolution(PrimaryConfig, CurrentShot, CameraOverscan);
 }
 
+bool UMoviePipelineDeferredPassBase::GetManualExposureEnabledForCamera(const int32 InCameraIndex) const
+{
+	UMoviePipelineExecutorShot* CurrentShot = GetPipeline()->GetActiveShotList()[GetPipeline()->GetCurrentShotIndex()];
+
+	if (GetNumCamerasToRender() == 1)
+	{
+		APlayerCameraManager* PlayerCameraManager = GetPipeline()->GetWorld()->GetFirstPlayerController()->PlayerCameraManager;
+
+		// By design, we don't use PlayerCameraManager->GetCameraCacheView() or PlayerCameraManager->ViewTarget.POV, as these cached
+		// values haven't been initialized at the beginning of the first tick where the movie passes are initialized.  Instead, we
+		// need to pull values from the component on the target actor.
+		if (PlayerCameraManager->ViewTarget.Target)
+		{
+			UCameraComponent* CameraComponent = MovieSceneHelpers::CameraComponentFromActor(PlayerCameraManager->ViewTarget.Target);
+			if (CameraComponent && CameraComponent->PostProcessBlendWeight > 0.0f && CameraComponent->PostProcessSettings.bOverride_AutoExposureMethod)
+			{
+				return CameraComponent->PostProcessSettings.AutoExposureMethod == AEM_Manual;
+			}
+		}
+	}
+	else
+	{
+		FMinimalViewInfo CameraViewInfo;
+		UCameraComponent* CameraComponent;
+		GetPipeline()->GetSidecarCameraData(CurrentShot, InCameraIndex, CameraViewInfo, &CameraComponent);
+
+		if (CameraViewInfo.PostProcessBlendWeight > 0.0f && CameraViewInfo.PostProcessSettings.bOverride_AutoExposureMethod)
+		{
+			return CameraViewInfo.PostProcessSettings.AutoExposureMethod == AEM_Manual;
+		}
+	}
+
+	return false;
+}
+
 FMoviePipelineRenderPassMetrics UMoviePipelineDeferredPassBase::GetRenderPassMetricsForCamera(const int32 InCameraIndex, const FMoviePipelineRenderPassMetrics& InSampleState) const
 {
 	// Add per-camera custom backbuffer size support here.
@@ -251,7 +286,7 @@ void UMoviePipelineDeferredPassBase::SetupImpl(const MoviePipeline::FMoviePipeli
 			CameraData.SystemMemoryMirror = FSceneViewStateInterface::SystemMemoryMirrorAllocate();
 		}
 
-		if (HighResSettings->bAutoExposureRenderPass && HighResTileCount > 1 && bRenderMainPass)
+		if (HighResTileCount > 1 && bRenderMainPass && !GetManualExposureEnabledForCamera(CamIndex))
 		{
 			// Add view state for auto-exposure render pass
 			FMultiCameraViewStateData::FPerTile& PerTile = CameraData.TileData.FindOrAdd(FIntPoint(-1, -1));
@@ -566,8 +601,11 @@ void UMoviePipelineDeferredPassBase::RenderSample_GameThreadImpl(const FMoviePip
 		// Main Render Pass
 		if (bRenderMainPass)
 		{
+			// We have an auto-exposure pass if we allocated a view state for it
+			bool bHasAutoExposurePass = CameraViewStateData.Num() && CameraViewStateData[0].TileData.Contains(FIntPoint(-1, -1));
+
 			// Auto-exposure pass is rendered only for the first tile and sample, actual pass is always rendered
-			bool bAutoExposurePass = InSampleState.bAutoExposureRenderPass && InSampleState.TileIndexes.X == 0 && InSampleState.TileIndexes.Y == 0 && InSampleState.TemporalSampleIndex == 0 && InSampleState.SpatialSampleIndex == 0;
+			bool bAutoExposurePass = bHasAutoExposurePass && InSampleState.TileIndexes.X == 0 && InSampleState.TileIndexes.Y == 0 && InSampleState.TemporalSampleIndex == 0 && InSampleState.SpatialSampleIndex == 0;
 			bool bActualPass = true;
 
 			while (bAutoExposurePass || bActualPass)
@@ -1126,8 +1164,8 @@ void UMoviePipelineDeferredPassBase::PostRendererSubmission(const FMoviePipeline
 
 bool UMoviePipelineDeferredPassBase::IsAutoExposureAllowed(const FMoviePipelineRenderPassMetrics& InSampleState) const
 {
-	// High-res tiling supports auto-exposure, if the auto exposure render pass is enabled.
-	return !(InSampleState.GetTileCount() > 1) || InSampleState.bAutoExposureRenderPass;
+	// For high-res tiling, check if a view state exists for auto-exposure, with key {-1,-1}
+	return (InSampleState.GetTileCount() == 1) || (CameraViewStateData.Num() && CameraViewStateData[0].TileData.Contains(FIntPoint(-1, -1)));
 }
 
 #if WITH_EDITOR
diff --git a/Engine/Plugins/MovieScene/MovieRenderPipeline/Source/MovieRenderPipelineRenderPasses/Public/MoviePipelineDeferredPasses.h b/Engine/Plugins/MovieScene/MovieRenderPipeline/Source/MovieRenderPipelineRenderPasses/Public/MoviePipelineDeferredPasses.h
--- a/Engine/Plugins/MovieScene/MovieRenderPipeline/Source/MovieRenderPipelineRenderPasses/Public/MoviePipelineDeferredPasses.h
+++ b/Engine/Plugins/MovieScene/MovieRenderPipeline/Source/MovieRenderPipelineRenderPasses/Public/MoviePipelineDeferredPasses.h
@@ -94,7 +94,8 @@ protected:
 	UE_DEPRECATED(5.6, "GetEffectiveOutputResolutionForCamera is deprecated. Use GetOverscannedResolutionForCamera instead")
 	virtual FIntPoint GetEffectiveOutputResolutionForCamera(const int32 InCameraIndex) const;
 	virtual FIntPoint GetOverscannedResolutionForCamera(const int32 InCameraIndex) const;
-	
+	bool GetManualExposureEnabledForCamera(const int32 InCameraIndex) const;
+
 	bool IsUsingDataLayers() const;
 	int32 GetNumStencilLayers() const;
 	TArray<FString> GetStencilLayerNames() const;
