Fixing static analysis warning

#preflihgt none
#rnx

diff --git a/Engine/Plugins/MovieScene/MovieRenderPipeline/Source/MovieRenderPipelineMP4Encoder/Private/Windows/MoviePipelineMP4Encoder.cpp b/Engine/Plugins/MovieScene/MovieRenderPipeline/Source/MovieRenderPipelineMP4Encoder/Private/Windows/MoviePipelineMP4Encoder.cpp
--- a/Engine/Plugins/MovieScene/MovieRenderPipeline/Source/MovieRenderPipelineMP4Encoder/Private/Windows/MoviePipelineMP4Encoder.cpp
+++ b/Engine/Plugins/MovieScene/MovieRenderPipeline/Source/MovieRenderPipelineMP4Encoder/Private/Windows/MoviePipelineMP4Encoder.cpp
@@ -164,14 +164,14 @@ bool FMoviePipelineMP4Encoder::WriteFrame(const uint8* InFrameData)
 	{
 		Result = MFCreateSample(&Sample);
 	}
-
-	if (SUCCEEDED(Result))
+	
+	if (SUCCEEDED(Result) && Sample)
 	{
 		Result = Sample->AddBuffer(Buffer);
 	}
 
 	// Duration & Timestamp for this frame.
-	if (SUCCEEDED(Result))
+	if (SUCCEEDED(Result) && Sample)
 	{
 		// In 100 nano-second units.
 		uint64 FrameDuration = (10 * 1000 * 1000) * Options.FrameRate.AsInterval();
@@ -179,7 +179,7 @@ bool FMoviePipelineMP4Encoder::WriteFrame(const uint8* InFrameData)
 		NumVideoSamplesWritten++;
 	}
 
-	if (SUCCEEDED(Result))
+	if (SUCCEEDED(Result) && Sample)
 	{
 		// In 100 nano-second units.
 		uint64 FrameDuration = (10 * 1000 * 1000) * Options.FrameRate.AsInterval();
@@ -187,7 +187,7 @@ bool FMoviePipelineMP4Encoder::WriteFrame(const uint8* InFrameData)
 	}
 
 	// Send the sample to the Sink Writer
-	if (SUCCEEDED(Result))
+	if (SUCCEEDED(Result) && Sample)
 	{
 		Result = SinkWriter->WriteSample(VideoStreamIndex, Sample);
 	}
@@ -267,20 +267,20 @@ bool FMoviePipelineMP4Encoder::WriteAudioSample(const TArrayView<int16>& InAudio
 			Result = MFCreateSample(&Sample);
 		}
 
-		if (SUCCEEDED(Result))
+		if (SUCCEEDED(Result) && Sample)
 		{
 			Result = Sample->AddBuffer(Buffer);
 		}
 
 		// Duration & Timestamp for this frame.
-		if (SUCCEEDED(Result))
+		if (SUCCEEDED(Result) && Sample)
 		{
 			// The duration matches the duration of one frame.
 			uint64 FrameDuration = (10 * 1000 * 1000) * Options.FrameRate.AsInterval();
 			Result = Sample->SetSampleDuration(FrameDuration);
 		}
 
-		if (SUCCEEDED(Result))
+		if (SUCCEEDED(Result) && Sample)
 		{
 			uint64 FrameDuration = (10 * 1000 * 1000) * Options.FrameRate.AsInterval();
 			Result = Sample->SetSampleTime(NumAudioSamplesWritten * FrameDuration);
@@ -291,7 +291,7 @@ bool FMoviePipelineMP4Encoder::WriteAudioSample(const TArrayView<int16>& InAudio
 		}
 
 		// Send the sample to the Sink Writer
-		if (SUCCEEDED(Result))
+		if (SUCCEEDED(Result) && Sample)
 		{
 			Result = SinkWriter->WriteSample(AudioStreamIndex, Sample);
 		}
@@ -743,7 +743,7 @@ bool FMoviePipelineMP4Encoder::InitializeEncoder()
 		Result = CreateVideoMediaOutputStream(&VideoMediaTypeOut, Options);
 	}
 	// Video Stream
-	if (SUCCEEDED(Result))
+	if (SUCCEEDED(Result) && SinkWriter)
 	{
 		Result = SinkWriter->AddStream(VideoMediaTypeOut, (DWORD*)&VideoStreamIndex);
 	}
@@ -754,7 +754,7 @@ bool FMoviePipelineMP4Encoder::InitializeEncoder()
 	}
 
 	// Audio Stream
-	if (SUCCEEDED(Result))
+	if (SUCCEEDED(Result) && SinkWriter)
 	{
 		Result = SinkWriter->AddStream(AudioMediaTypeOut, (DWORD*)&AudioStreamIndex);
 	}
@@ -769,7 +769,7 @@ bool FMoviePipelineMP4Encoder::InitializeEncoder()
 		Result = CreateVideoMediaTypeIn(&VideoMediaTypeIn, Options);
 	}
 
-	if (SUCCEEDED(Result))
+	if (SUCCEEDED(Result) && SinkWriter)
 	{
 		IMFAttributes* pEncAttrs = nullptr;
 		Result = GetVideoStreamEncoderAttributes(&pEncAttrs, Options);
@@ -784,21 +784,24 @@ bool FMoviePipelineMP4Encoder::InitializeEncoder()
 		}
 	}
 
-	if (SUCCEEDED(Result))
+	if (SUCCEEDED(Result) && SinkWriter)
 	{
 		Result = SinkWriter->SetInputMediaType(AudioStreamIndex, AudioMediaTypeIn, NULL);
 	}
 
-	ICodecAPI* codecApi;
-	Result = SinkWriter->GetServiceForStream(VideoStreamIndex, GUID_NULL, __uuidof(ICodecAPI), (LPVOID*)&codecApi);
+	if (SUCCEEDED(Result) && SinkWriter)
+	{
+		ICodecAPI* codecApi = nullptr;
+		Result = SinkWriter->GetServiceForStream(VideoStreamIndex, GUID_NULL, __uuidof(ICodecAPI), (LPVOID*)&codecApi);
+	}
 
 	// The sink writer can now accept data.
-	if (SUCCEEDED(Result))
+	if (SUCCEEDED(Result) && SinkWriter)
 	{
 		SinkWriter->BeginWriting();
 	}
 
-	if (SUCCEEDED(Result))
+	if (SUCCEEDED(Result) && SinkWriter)
 	{
 		SinkWriter->AddRef();
 	}
