MoviePipeline: Added additional diagnostic information to EXR metadata.

diff --git a/Engine/Plugins/MovieScene/MovieRenderPipeline/Source/MovieRenderPipelineCore/Private/MoviePipelineRendering.cpp b/Engine/Plugins/MovieScene/MovieRenderPipeline/Source/MovieRenderPipelineCore/Private/MoviePipelineRendering.cpp
--- a/Engine/Plugins/MovieScene/MovieRenderPipeline/Source/MovieRenderPipelineCore/Private/MoviePipelineRendering.cpp
+++ b/Engine/Plugins/MovieScene/MovieRenderPipeline/Source/MovieRenderPipelineCore/Private/MoviePipelineRendering.cpp
@@ -266,6 +266,9 @@ void UMoviePipeline::RenderFrame()
 
 		UE::MoviePipeline::GetHardwareUsageMetadata(CachedOutputState.FileMetadata, ResolvedOutputDirectory);
 
+		// Add in additional diagnostic information (engine version, etc)
+		constexpr bool bIsGraph = false;
+		UE::MoviePipeline::GetDiagnosticMetadata(CachedOutputState.FileMetadata, bIsGraph);
 
 		// We'll leave these in for legacy, when this tracks the 'Main' camera (of the player), render passes that support
 		// multiple cameras will have to write each camera name into their metadata.
diff --git a/Engine/Plugins/MovieScene/MovieRenderPipeline/Source/MovieRenderPipelineCore/Private/MoviePipelineUtils.cpp b/Engine/Plugins/MovieScene/MovieRenderPipeline/Source/MovieRenderPipelineCore/Private/MoviePipelineUtils.cpp
--- a/Engine/Plugins/MovieScene/MovieRenderPipeline/Source/MovieRenderPipelineCore/Private/MoviePipelineUtils.cpp
+++ b/Engine/Plugins/MovieScene/MovieRenderPipeline/Source/MovieRenderPipelineCore/Private/MoviePipelineUtils.cpp
@@ -26,6 +26,7 @@
 #include "LevelSequence.h"
 #include "MovieRenderPipelineCoreModule.h"
 #include "Math/Halton.h"
+#include "Misc/EngineVersion.h"
 #include "Misc/Paths.h"
 #include "Misc/FrameRate.h"
 #include "MoviePipelineQueue.h"
@@ -1085,6 +1086,12 @@ namespace UE
 			InFileMetadata.Add(TEXT("unreal/jobComment"), JobComment);
 		}
 
+		void GetDiagnosticMetadata(TMap<FString, FString>& InFileMetadata, const bool bIsGraph)
+		{
+			InFileMetadata.Add(TEXT("unreal/build"), FEngineVersion::Current().ToString());
+			InFileMetadata.Add(TEXT("unreal/usesGraph"), LexToString(bIsGraph));
+		}
+
 		void GetCachedGPUDriverInfo(TMap<FString, FString>& InFileMetadata)
 		{
 			static TMap<FString, FString> CachedInfo;
diff --git a/Engine/Plugins/MovieScene/MovieRenderPipeline/Source/MovieRenderPipelineCore/Public/MoviePipelineUtils.h b/Engine/Plugins/MovieScene/MovieRenderPipeline/Source/MovieRenderPipelineCore/Public/MoviePipelineUtils.h
--- a/Engine/Plugins/MovieScene/MovieRenderPipeline/Source/MovieRenderPipelineCore/Public/MoviePipelineUtils.h
+++ b/Engine/Plugins/MovieScene/MovieRenderPipeline/Source/MovieRenderPipelineCore/Public/MoviePipelineUtils.h
@@ -189,6 +189,7 @@ namespace UE
 		MOVIERENDERPIPELINECORE_API FString GetJobAuthor(const UMoviePipelineExecutorJob* InJob);
 		MOVIERENDERPIPELINECORE_API void GetSharedFormatArguments(TMap<FString, FString>& InFilenameArguments, TMap<FString, FString>& InFileMetadata, const FDateTime& InDateTime, const int32 InVersionNumber, const UMoviePipelineExecutorJob* InJob, const FTimespan& InInitializationTimeOffset = FTimespan());
 		MOVIERENDERPIPELINECORE_API void GetHardwareUsageMetadata(TMap<FString, FString>& InFileMetadata, const FString& InOutputDir);
+		MOVIERENDERPIPELINECORE_API void GetDiagnosticMetadata(TMap<FString, FString>& InFileMetadata, const bool bIsGraph);
 		MOVIERENDERPIPELINECORE_API void GetMetadataFromCineCamera(class UCineCameraComponent* InComponent, const FString& InCameraName, const FString& InRenderPassName, TMap<FString, FString>& InOutMetadata);
 		MOVIERENDERPIPELINECORE_API void GetMetadataFromCameraLocRot(const FString& InCameraName, const FString& InRenderPassName, const FVector& InCurLoc, const FRotator& InCurRot, const FVector& InPrevLoc, const FRotator& InPrevRot, TMap<FString, FString>& InOutMetadata);
 		MOVIERENDERPIPELINECORE_API FMoviePipelineRenderPassMetrics GetRenderPassMetrics(UMoviePipelinePrimaryConfig* InPrimaryConfig, UMoviePipelineExecutorShot* InPipelineExecutorShot, const FMoviePipelineRenderPassMetrics& InRenderPassMetrics, const FIntPoint& InEffectiveOutputResolution);
diff --git a/Engine/Plugins/MovieScene/MovieRenderPipeline/Source/MovieRenderPipelineRenderPasses/Private/MovieGraphImageSequenceOutputNode.cpp b/Engine/Plugins/MovieScene/MovieRenderPipeline/Source/MovieRenderPipelineRenderPasses/Private/MovieGraphImageSequenceOutputNode.cpp
--- a/Engine/Plugins/MovieScene/MovieRenderPipeline/Source/MovieRenderPipelineRenderPasses/Private/MovieGraphImageSequenceOutputNode.cpp
+++ b/Engine/Plugins/MovieScene/MovieRenderPipeline/Source/MovieRenderPipelineRenderPasses/Private/MovieGraphImageSequenceOutputNode.cpp
@@ -306,8 +306,10 @@ TUniquePtr<FEXRImageWriteTask> UMovieGraphImageSequenceOutputNode_EXR::CreateIma
 
 void UMovieGraphImageSequenceOutputNode_EXR::PrepareTaskGlobalMetadata(FEXRImageWriteTask& InOutImageTask, UE::MovieGraph::FMovieGraphOutputMergerFrame* InRawFrameData, TMap<FString, FString>& InMetadata) const
 {
-	// Add in hardware usage metadata
+	// Add in hardware usage & diagnostic metadata
+	constexpr bool bIsGraph = true;
 	UE::MoviePipeline::GetHardwareUsageMetadata(InMetadata, FPaths::GetPath(InOutImageTask.Filename));
+	UE::MoviePipeline::GetDiagnosticMetadata(InMetadata, bIsGraph);
 
 	// Add passed in resolved metadata
 	for (const TPair<FString, FString>& Metadata : InMetadata)
