MoviePipeline: Minor cosmetic polish for the H.264 MP4 node and MRQ setting.

Also fixed some other minor bugs in the H.264 node.

diff --git a/Engine/Plugins/MovieScene/MovieRenderPipeline/Source/MovieRenderPipelineCore/Public/MoviePipelinePostRenderSettings.h b/Engine/Plugins/MovieScene/MovieRenderPipeline/Source/MovieRenderPipelineCore/Public/MoviePipelinePostRenderSettings.h
--- a/Engine/Plugins/MovieScene/MovieRenderPipeline/Source/MovieRenderPipelineCore/Public/MoviePipelinePostRenderSettings.h
+++ b/Engine/Plugins/MovieScene/MovieRenderPipeline/Source/MovieRenderPipelineCore/Public/MoviePipelinePostRenderSettings.h
@@ -181,7 +181,7 @@ struct MOVIERENDERPIPELINECORE_API FMovieGraphPostRenderSettings
 	FMovieGraphPostRenderVideoPlayOptions PostRenderAvidDNxHRPlayOptions;
 
 	/** These options are used to determine how to play H.264 MP4 movies. */
-	UPROPERTY(EditAnywhere, BlueprintReadWrite, Category = "Play Render Output", DisplayName = "H.264 MP4 Movies")
+	UPROPERTY(EditAnywhere, BlueprintReadWrite, Category = "Play Render Output (Quick Render only)", DisplayName = "H.264 MP4 Movies")
 	FMovieGraphPostRenderVideoPlayOptions PostRenderMP4PlayOptions;
 
 	// TODO: No CLE support yet.
diff --git a/Engine/Plugins/MovieScene/MovieRenderPipeline/Source/MovieRenderPipelineMP4Encoder/Private/Graph/MovieGraphMP4EncoderNode.cpp b/Engine/Plugins/MovieScene/MovieRenderPipeline/Source/MovieRenderPipelineMP4Encoder/Private/Graph/MovieGraphMP4EncoderNode.cpp
--- a/Engine/Plugins/MovieScene/MovieRenderPipeline/Source/MovieRenderPipelineMP4Encoder/Private/Graph/MovieGraphMP4EncoderNode.cpp
+++ b/Engine/Plugins/MovieScene/MovieRenderPipeline/Source/MovieRenderPipelineMP4Encoder/Private/Graph/MovieGraphMP4EncoderNode.cpp
@@ -22,7 +22,7 @@ UMovieGraphMP4EncoderNode::UMovieGraphMP4EncoderNode()
 #if WITH_EDITOR
 FText UMovieGraphMP4EncoderNode::GetNodeTitle(const bool bGetDescriptive) const
 {
-	static const FText NodeName = NSLOCTEXT("MovieGraphNodes", "NodeName_MP4", "h264 MP4");
+	static const FText NodeName = NSLOCTEXT("MovieGraphNodes", "NodeName_MP4", "H.264 MP4");
 	return NodeName;
 }
 
@@ -72,13 +72,13 @@ TUniquePtr<MovieRenderGraph::IVideoCodecWriter> UMovieGraphMP4EncoderNode::Initi
 	Options.Width = InInitializationContext.Resolution.X;
 	Options.Height = InInitializationContext.Resolution.Y;
 	Options.FrameRate = EffectiveFrameRate;
-	Options.CommonMaxBitRate = FMath::RoundToInt(MaxBitrateInMbps*1024*1024);
-	Options.CommonMeanBitRate = FMath::RoundToInt(AverageBitrateInMbps*1024*1024);
+	Options.CommonMaxBitRate = FMath::RoundToInt(EvaluatedNode->MaxBitrateInMbps*1024*1024);
+	Options.CommonMeanBitRate = FMath::RoundToInt(EvaluatedNode->AverageBitrateInMbps*1024*1024);
 	Options.CommonQualityVsSpeed = 100;
-	Options.CommonConstantRateFactor = ConstantRateFactor;
-	Options.EncodingProfile = EncodingProfile;
-	Options.EncodingLevel = EncodingLevel;
-	Options.EncodingRateControl = EncodingRateControl;
+	Options.CommonConstantRateFactor = EvaluatedNode->ConstantRateFactor;
+	Options.EncodingProfile = EvaluatedNode->EncodingProfile;
+	Options.EncodingLevel = EvaluatedNode->EncodingLevel;
+	Options.EncodingRateControl = EvaluatedNode->EncodingRateControl;
 	
 	TUniquePtr<FMP4CodecWriter> NewWriter = MakeUnique<FMP4CodecWriter>();
 	NewWriter->Writer = MakeUnique<FMoviePipelineMP4Encoder>(Options);
diff --git a/Engine/Plugins/MovieScene/MovieRenderPipeline/Source/MovieRenderPipelineMP4Encoder/Private/Windows/MoviePipelineMP4Encoder.cpp b/Engine/Plugins/MovieScene/MovieRenderPipeline/Source/MovieRenderPipelineMP4Encoder/Private/Windows/MoviePipelineMP4Encoder.cpp
--- a/Engine/Plugins/MovieScene/MovieRenderPipeline/Source/MovieRenderPipelineMP4Encoder/Private/Windows/MoviePipelineMP4Encoder.cpp
+++ b/Engine/Plugins/MovieScene/MovieRenderPipeline/Source/MovieRenderPipelineMP4Encoder/Private/Windows/MoviePipelineMP4Encoder.cpp
@@ -789,11 +789,13 @@ bool FMoviePipelineMP4Encoder::InitializeEncoder()
 		Result = SinkWriter->SetInputMediaType(AudioStreamIndex, AudioMediaTypeIn, NULL);
 	}
 
-	if (SUCCEEDED(Result) && SinkWriter)
-	{
-		ICodecAPI* codecApi = nullptr;
-		Result = SinkWriter->GetServiceForStream(VideoStreamIndex, GUID_NULL, __uuidof(ICodecAPI), (LPVOID*)&codecApi);
-	}
+	// Note: This has no practical use to the user. However, for diagnostic purposes, `codecApi` can be inspected in the debugger to determine which
+	// encoder is being used (software or hardware).
+	// if (SUCCEEDED(Result) && SinkWriter)
+	// {
+	// 	ICodecAPI* codecApi = nullptr;
+	// 	Result = SinkWriter->GetServiceForStream(VideoStreamIndex, GUID_NULL, __uuidof(ICodecAPI), (LPVOID*)&codecApi);
+	// }
 
 	// The sink writer can now accept data.
 	if (SUCCEEDED(Result) && SinkWriter)
diff --git a/Engine/Plugins/MovieScene/MovieRenderPipeline/Source/MovieRenderPipelineMP4Encoder/Public/Graph/MovieGraphMP4EncoderNode.h b/Engine/Plugins/MovieScene/MovieRenderPipeline/Source/MovieRenderPipelineMP4Encoder/Public/Graph/MovieGraphMP4EncoderNode.h
--- a/Engine/Plugins/MovieScene/MovieRenderPipeline/Source/MovieRenderPipelineMP4Encoder/Public/Graph/MovieGraphMP4EncoderNode.h
+++ b/Engine/Plugins/MovieScene/MovieRenderPipeline/Source/MovieRenderPipelineMP4Encoder/Public/Graph/MovieGraphMP4EncoderNode.h
@@ -51,6 +51,9 @@ protected:
 	TWeakObjectPtr<UMovieGraphPipeline> CachedPipeline;
 
 public:
+	UPROPERTY(EditAnywhere, BlueprintReadWrite, Category = Overrides, meta = (InlineEditConditionToggle))
+	uint8 bOverride_EncodingRateControl : 1;
+	
 	UPROPERTY(EditAnywhere, BlueprintReadWrite, Category = Overrides, meta = (InlineEditConditionToggle))
 	uint8 bOverride_AverageBitrateInMbps : 1;
 	
@@ -77,7 +80,7 @@ public:
 	* to the filesize, while Average/Constant modes allow you to suggest desired sizes, though the resulting file may still
 	* end up larger or smaller than specified.
 	*/
-	UPROPERTY(BlueprintReadWrite, EditAnywhere, Category = "Settings")
+	UPROPERTY(BlueprintReadWrite, EditAnywhere, Category = "H.264", meta=(EditCondition="bOverride_EncodingRateControl"))
 	EMoviePipelineMP4EncodeRateControlMode EncodingRateControl = EMoviePipelineMP4EncodeRateControlMode::Quality;
 	
 	/**
@@ -87,7 +90,7 @@ public:
 	* 
 	* Only applies to encoding modes not related to Quality.
 	*/
-	UPROPERTY(BlueprintReadWrite, EditAnywhere, Category = "Settings", meta=(ClampMin=0.1, UIMin=0.1, UIMax=50, EditCondition="bOverride_AverageBitrateInMbps"))
+	UPROPERTY(BlueprintReadWrite, EditAnywhere, Category = "H.264", meta=(ClampMin=0.1, UIMin=0.1, UIMax=50, EditCondition="bOverride_AverageBitrateInMbps"))
 	float AverageBitrateInMbps = 8.f;
 	
 	/**
@@ -97,7 +100,7 @@ public:
 	*
 	* Only applies to constrained max bit rate.
 	*/
-	UPROPERTY(BlueprintReadWrite, EditAnywhere, Category = "Settings", meta = (ClampMin = 0.1, UIMin = 0.1, UIMax = 50, EditCondition = "bOverride_MaxBitrateInMbps"))
+	UPROPERTY(BlueprintReadWrite, EditAnywhere, Category = "H.264", meta = (ClampMin = 0.1, UIMin = 0.1, UIMax = 50, EditCondition = "bOverride_MaxBitrateInMbps"))
 	float MaxBitrateInMbps = 16.f;
 
 	/**
@@ -108,18 +111,18 @@ public:
 	*
 	* Only applies to Quality encoding rate control mode.
 	*/
-	UPROPERTY(BlueprintReadWrite, EditAnywhere, Category = "Settings", meta=(ClampMin=16, UIMin=16, ClampMax=51, UIMax=51, EditCondition="bOverride_ConstantRateFactor"))
+	UPROPERTY(BlueprintReadWrite, EditAnywhere, Category = "H.264", meta=(ClampMin=16, UIMin=16, ClampMax=51, UIMax=51, EditCondition="bOverride_ConstantRateFactor"))
 	int32 ConstantRateFactor = 20;
 
 	/** A higher profile generally results in a better quality video for the same bitrate, but may not be supported for playback on old devices. */
-	UPROPERTY(BlueprintReadWrite, EditAnywhere, Category = "Settings", meta=(AdvancedDisplay, EditCondition="bOverride_EncodingProfile"))
+	UPROPERTY(BlueprintReadWrite, EditAnywhere, Category = "H.264", meta=(AdvancedDisplay, EditCondition="bOverride_EncodingProfile"))
 	EMoviePipelineMP4EncodeProfile EncodingProfile = EMoviePipelineMP4EncodeProfile::High;
 
 	/** 
 	* A higher encode level generally results in a better quality video for the same bitrate, but may not be supported for playback on old devices.
 	* Auto will let the encoder automatically select the best level given the resolution, profile, and bitrate.
 	*/
-	UPROPERTY(BlueprintReadWrite, EditAnywhere, Category = "Settings", meta=(AdvancedDisplay, EditCondition="bOverride_EncodingLevel"))
+	UPROPERTY(BlueprintReadWrite, EditAnywhere, Category = "H.264", meta=(AdvancedDisplay, EditCondition="bOverride_EncodingLevel"))
 	EMoviePipelineMP4EncodeLevel EncodingLevel = EMoviePipelineMP4EncodeLevel::Auto;
 	
 	/**
diff --git a/Engine/Plugins/MovieScene/MovieRenderPipeline/Source/MovieRenderPipelineMP4Encoder/Public/MoviePipelineMP4EncoderOutput.h b/Engine/Plugins/MovieScene/MovieRenderPipeline/Source/MovieRenderPipelineMP4Encoder/Public/MoviePipelineMP4EncoderOutput.h
--- a/Engine/Plugins/MovieScene/MovieRenderPipeline/Source/MovieRenderPipelineMP4Encoder/Public/MoviePipelineMP4EncoderOutput.h
+++ b/Engine/Plugins/MovieScene/MovieRenderPipeline/Source/MovieRenderPipelineMP4Encoder/Public/MoviePipelineMP4EncoderOutput.h
@@ -35,7 +35,7 @@ protected:
 
 	// UMoviePipelineOutputBase Interface
 #if WITH_EDITOR
-	virtual FText GetDisplayText() const override { return NSLOCTEXT("MovieRenderPipeline", "MP4EncoderNode_DisplayName", "h264 MP4 [8bit]"); }
+	virtual FText GetDisplayText() const override { return NSLOCTEXT("MovieRenderPipeline", "MP4EncoderNode_DisplayName", "H.264 MP4 [8bit]"); }
 #endif
 	// ~UMoviePipelineOutputBase Interface
 
