MoviePipeline: When drag-dropping a branch pin from a Render Layer node onto the Outputs or Inputs node, the new member created is now named based on the layer name in the Render Layer node. Previously the new input or output created was given a generic name like "NewInput" or "NewOutput".

diff --git a/Engine/Plugins/MovieScene/MovieRenderPipeline/Source/MovieRenderPipelineEditor/Private/Graph/MovieGraphSchema.cpp b/Engine/Plugins/MovieScene/MovieRenderPipeline/Source/MovieRenderPipelineEditor/Private/Graph/MovieGraphSchema.cpp
--- a/Engine/Plugins/MovieScene/MovieRenderPipeline/Source/MovieRenderPipelineEditor/Private/Graph/MovieGraphSchema.cpp
+++ b/Engine/Plugins/MovieScene/MovieRenderPipeline/Source/MovieRenderPipelineEditor/Private/Graph/MovieGraphSchema.cpp
@@ -6,6 +6,7 @@
 #include "Graph/MovieGraphConfig.h"
 #include "Graph/Nodes/MovieGraphInputNode.h"
 #include "Graph/Nodes/MovieGraphOutputNode.h"
+#include "Graph/Nodes/MovieGraphRenderLayerNode.h"
 #include "Graph/Nodes/MovieGraphRerouteNode.h"
 #include "Graph/Nodes/MovieGraphSubgraphNode.h"
 #include "Graph/Nodes/MovieGraphVariableNode.h"
@@ -53,6 +54,11 @@ namespace UE::MovieGraph::Private
 {
 	UMovieGraphNode* GetGraphNodeFromEdPin(const UEdGraphPin* InPin)
 	{
+		if (!InPin)
+		{
+			return nullptr;
+		}
+		
 		const UMoviePipelineEdGraphNodeBase* EdGraphNode = CastChecked<UMoviePipelineEdGraphNodeBase>(InPin->GetOwningNode());
 
 		UMovieGraphNode* RuntimeNode = EdGraphNode->GetRuntimeNode();
@@ -162,7 +168,16 @@ UEdGraphPin* UMovieGraphSchema::DropPinOnNode(UEdGraphNode* InTargetNode, const
 		FText NewMemberName;
 		if (InSourcePinName == NAME_None)
 		{
-			NewMemberName = (InSourcePinDirection == EGPD_Input) ? LOCTEXT("NewInputName", "NewInput") : LOCTEXT("NewOutputName", "NewOutput");
+			// If the source of the connection is a Render Layer node, then name the new member based on the layer name. Otherwise, give the new
+			// member a generic name.
+			if (const UMovieGraphRenderLayerNode* RenderLayerNode = Cast<UMovieGraphRenderLayerNode>(UE::MovieGraph::Private::GetGraphNodeFromEdPin(PinBeingDropped)))
+			{
+				NewMemberName = FText::FromString(RenderLayerNode->LayerName);
+			}
+			else
+			{
+				NewMemberName = (InSourcePinDirection == EGPD_Input) ? LOCTEXT("NewInputName", "NewInput") : LOCTEXT("NewOutputName", "NewOutput");
+			}
 		}
 		else
 		{
@@ -196,6 +211,11 @@ UEdGraphPin* UMovieGraphSchema::DropPinOnNode(UEdGraphNode* InTargetNode, const
 	return NewEdPin; 
 }
 
+void UMovieGraphSchema::SetPinBeingDroppedOnNode(UEdGraphPin* InSourcePin) const
+{
+	PinBeingDropped = InSourcePin;
+}
+
 bool UMovieGraphSchema::SafeDeleteNodeFromGraph(UEdGraph* Graph, UEdGraphNode* Node) const
 {
 	if (!Graph || !Node)
diff --git a/Engine/Plugins/MovieScene/MovieRenderPipeline/Source/MovieRenderPipelineEditor/Private/Graph/MovieGraphSchema.h b/Engine/Plugins/MovieScene/MovieRenderPipeline/Source/MovieRenderPipelineEditor/Private/Graph/MovieGraphSchema.h
--- a/Engine/Plugins/MovieScene/MovieRenderPipeline/Source/MovieRenderPipelineEditor/Private/Graph/MovieGraphSchema.h
+++ b/Engine/Plugins/MovieScene/MovieRenderPipeline/Source/MovieRenderPipelineEditor/Private/Graph/MovieGraphSchema.h
@@ -29,6 +29,7 @@ public:
 	virtual bool ShouldHidePinDefaultValue(UEdGraphPin* Pin) const override;
 	virtual bool SupportsDropPinOnNode(UEdGraphNode* InTargetNode, const FEdGraphPinType& InSourcePinType, EEdGraphPinDirection InSourcePinDirection, FText& OutErrorMessage) const override;
 	virtual UEdGraphPin* DropPinOnNode(UEdGraphNode* InTargetNode, const FName& InSourcePinName, const FEdGraphPinType& InSourcePinType, EEdGraphPinDirection InSourcePinDirection) const override;
+	virtual void SetPinBeingDroppedOnNode(UEdGraphPin* InSourcePin) const override;
 	virtual bool SafeDeleteNodeFromGraph(UEdGraph* Graph, UEdGraphNode* Node) const override;
 	//~ End EdGraphSchema Interface
 	
@@ -81,6 +82,9 @@ public:
 	
 private:
 	static TArray<UClass*> MoviePipelineNodeClasses;
+
+	/** The pin that is being dropped on another node during a drag-drop operation. Controlled by SupportsDropPinOnNode(). */
+	mutable UEdGraphPin* PinBeingDropped = nullptr;
 };
 
 /** Base class for schema actions in the graph. */
