MoviePipeline: Updated UX for changing the configuration type in the queue.

diff --git a/Engine/Plugins/MovieScene/MovieRenderPipeline/Source/MovieRenderPipelineEditor/Private/Widgets/SMoviePipelineQueueEditor.cpp b/Engine/Plugins/MovieScene/MovieRenderPipeline/Source/MovieRenderPipelineEditor/Private/Widgets/SMoviePipelineQueueEditor.cpp
--- a/Engine/Plugins/MovieScene/MovieRenderPipeline/Source/MovieRenderPipelineEditor/Private/Widgets/SMoviePipelineQueueEditor.cpp
+++ b/Engine/Plugins/MovieScene/MovieRenderPipeline/Source/MovieRenderPipelineEditor/Private/Widgets/SMoviePipelineQueueEditor.cpp
@@ -407,6 +407,8 @@ public:
 				}
 			}
 		}
+
+		FSlateApplication::Get().DismissAllMenus();
 	}
 
 	void OnCreateNewGraphAndAssign()
@@ -427,12 +429,7 @@ public:
 					UMoviePipelineExecutorJob* Job = JobTreeItem->WeakJob.Get();
 					if (Job)
 					{
-						// We only apply to jobs already using graph configurations because they have to have been converted to
-						// a graph, and now they're replacing their configuration with a new one.
-						if (Job->IsUsingGraphConfiguration())
-						{
-							Job->SetGraphPreset(NewGraph);
-						}
+						Job->SetGraphPreset(NewGraph);
 					}
 				}
 			}
@@ -458,6 +455,8 @@ public:
 				}
 			}
 		}
+
+		FSlateApplication::Get().DismissAllMenus();
 	}
 
 	EVisibility GetPrimaryConfigModifiedVisibility() const
@@ -606,21 +605,6 @@ public:
 			AssetPickerConfig.OnAssetSelected = InOnAssetSelected;
 		}
 
-		if (TargetJob->IsUsingGraphConfiguration())
-		{
-			MenuBuilder.BeginSection(NAME_None, LOCTEXT("NewConfig_MenuSection", "New Configuration"));
-			{
-				MenuBuilder.AddMenuEntry(
-					LOCTEXT("NewGraph_Label", "New Graph"),
-					LOCTEXT("NewGraph_Tooltip", "Creates a new graph asset and assigns it to this job."),
-					FSlateIcon(),
-					FUIAction(InCreateNewGraphAndAssign),
-					NAME_None,
-					EUserInterfaceActionType::Button
-				);
-			}
-		}
-
 		MenuBuilder.BeginSection(NAME_None, LOCTEXT("CurrentConfig_MenuSection", "Current Configuration"));
 		{
 			if (!TargetJob->IsUsingGraphConfiguration())
@@ -633,19 +617,6 @@ public:
 					NAME_None,
 					EUserInterfaceActionType::Button
 				);
-
-				// Only primary jobs can do a Replace with Graph. Shots cannot have a graph config with a parent legacy config.
-				if (!bIsShot)
-				{
-					MenuBuilder.AddMenuEntry(
-						LOCTEXT("ReplaceWithGraph_Label", "Replace with Graph (Beta)"),
-						LOCTEXT("ReplaceWithGraph_Tooltip", "Replaces the current configuration with a new graph representation."),
-						FSlateIcon(),
-						FUIAction(InNewRenderGraph),
-						NAME_None,
-						EUserInterfaceActionType::Button
-					);
-				}
 			}
 			else
 			{
@@ -661,18 +632,65 @@ public:
 						EUserInterfaceActionType::Button
 					);
 				}
-				else
+			}
+		}
+		MenuBuilder.EndSection();
+
+		if (!bIsShot)
+		{
+			MenuBuilder.BeginSection(NAME_None, LOCTEXT("ConfigurationType_MenuSection", "Configuration Type"));
+			{
+				FUIAction SwitchToPresetAction;
+				SwitchToPresetAction.ExecuteAction = InClearGraph;
+				SwitchToPresetAction.GetActionCheckState = FGetActionCheckState::CreateLambda([TargetJob]() -> ECheckBoxState
 				{
-					MenuBuilder.AddMenuEntry(
-						LOCTEXT("ReplaceWithPreset_Label", "Replace with Preset"),
-						LOCTEXT("ReplaceWithPreset_Tooltip", "Replaces the current configuration with a new default (non-graph) config."),
-						FSlateIcon(),
-						FUIAction(InClearGraph),
-						NAME_None,
-						EUserInterfaceActionType::Button
-					);
-				}
+					return (TargetJob.IsValid() && TargetJob->IsUsingGraphConfiguration()) ? ECheckBoxState::Unchecked : ECheckBoxState::Checked;
+				});
+
+				FUIAction SwitchToGraphAction;
+				SwitchToGraphAction.ExecuteAction = InNewRenderGraph;
+				SwitchToGraphAction.GetActionCheckState = FGetActionCheckState::CreateLambda([TargetJob]() -> ECheckBoxState
+				{
+					return (TargetJob.IsValid() && TargetJob->IsUsingGraphConfiguration()) ? ECheckBoxState::Checked : ECheckBoxState::Unchecked;
+				});
+			
+				MenuBuilder.AddMenuEntry(
+					LOCTEXT("ConfigurationType_PresetLabel", "Preset"),
+					LOCTEXT("ConfigurationType_PresetTooltip", "Updates this job or shot to use a preset-based configuration."),
+					FSlateIcon(),
+					SwitchToPresetAction,
+					NAME_None,
+					EUserInterfaceActionType::RadioButton
+				);
+
+				MenuBuilder.AddMenuEntry(
+					LOCTEXT("ConfigurationType_GraphLabel", "Movie Render Graph (Beta)"),
+					LOCTEXT("ConfigurationType_GraphTooltip", "Updates this job or shot to use a graph-based configuration."),
+					FSlateIcon(),
+					SwitchToGraphAction,
+					NAME_None,
+					EUserInterfaceActionType::RadioButton
+				);
+			}
+			MenuBuilder.EndSection();
+		}
+
+		// If this menu is being created for a shot, only allow creating a new graph via the menu if the parent job is using a graph. Creating a
+		// graph for a shot with a parent using a preset-based configuration is not valid.
+		if (!bIsShot || (bIsShot && TargetJob->IsUsingGraphConfiguration()))
+		{
+			MenuBuilder.BeginSection(NAME_None, LOCTEXT("CreateNewAsset_MenuSection", "Create New Asset"));
+			{
+				MenuBuilder.AddMenuEntry(
+					LOCTEXT("NewGraph_Label", "Movie Render Graph"),
+					LOCTEXT("NewGraph_Tooltip", "Creates a new graph asset and assigns it to this job or shot."),
+					FSlateIcon(),
+					FUIAction(InCreateNewGraphAndAssign),
+					NAME_None,
+					EUserInterfaceActionType::Button
+				);
 			}
+			MenuBuilder.EndSection();
 		}
 
 		MenuBuilder.BeginSection(NAME_None, LOCTEXT("ImportConfig_MenuSection", "Import Configuration"));
