MoviePipeline: Added the ablilty in MRG to create a reroute node on an existing connection by double-clicking anywhere on the connection.

diff --git a/Engine/Plugins/MovieScene/MovieRenderPipeline/Source/MovieRenderPipelineEditor/Private/Graph/MovieGraphSchema.cpp b/Engine/Plugins/MovieScene/MovieRenderPipeline/Source/MovieRenderPipelineEditor/Private/Graph/MovieGraphSchema.cpp
--- a/Engine/Plugins/MovieScene/MovieRenderPipeline/Source/MovieRenderPipelineEditor/Private/Graph/MovieGraphSchema.cpp
+++ b/Engine/Plugins/MovieScene/MovieRenderPipeline/Source/MovieRenderPipelineEditor/Private/Graph/MovieGraphSchema.cpp
@@ -235,6 +235,38 @@ bool UMovieGraphSchema::SafeDeleteNodeFromGraph(UEdGraph* Graph, UEdGraphNode* N
 	return true;
 }
 
+void UMovieGraphSchema::OnPinConnectionDoubleCicked(UEdGraphPin* PinA, UEdGraphPin* PinB, const FVector2f& GraphPosition) const
+{
+	const FScopedTransaction Transaction(LOCTEXT("CreateRerouteNodeOnWire", "Create Reroute Node"));
+
+	// This is a temporary action for creating the reroute node, so the category, display name, tooltip, and keywords can just be empty.
+	constexpr int32 Grouping = 0;
+	const TSharedPtr<FMovieGraphSchemaAction> RerouteNodeAction =
+		MakeShared<FMovieGraphSchemaAction_NewNode>(FText::GetEmpty(), FText::GetEmpty(), FText::GetEmpty(), Grouping, FText::GetEmpty()); 
+	RerouteNodeAction->NodeClass = UMovieGraphRerouteNode::StaticClass();
+
+	// The node should be centered on where the mouse is clicked. The node's position is based on its top-left corner, so offset by
+	// half of its width/height to center it on the mouse. Note that the spacer size is duplicated from inside SGraphNodeKnot.
+	const FVector2f NodeSpacerSize(42.0f, 24.0f);
+	const FVector2f NodePosition = GraphPosition - (NodeSpacerSize * 0.5f);
+	
+	UEdGraph* OwningGraph = PinA->GetOwningNode()->GetGraph();
+	if (ensure(OwningGraph))
+	{
+		// Break the existing connection (ie, the connection that was clicked on).
+		PinA->BreakLinkTo(PinB);
+
+		// Create the reroute node. Using PinB as the "from" pin will create the downstream connection.
+		const UEdGraphNode* NewRerouteNode = RerouteNodeAction->PerformAction(OwningGraph, PinB, NodePosition);
+
+		// Connect the reroute node's left (input) pin back to the previously-connected upstream node.
+		if (UEdGraphPin* InputPin = NewRerouteNode->FindPin(NAME_None, EGPD_Input))
+		{
+			InputPin->MakeLinkTo(PinA);
+		}
+	}
+}
+
 void UMovieGraphSchema::InitMoviePipelineNodeClasses()
 {
 	if (MoviePipelineNodeClasses.Num() > 0)
diff --git a/Engine/Plugins/MovieScene/MovieRenderPipeline/Source/MovieRenderPipelineEditor/Private/Graph/MovieGraphSchema.h b/Engine/Plugins/MovieScene/MovieRenderPipeline/Source/MovieRenderPipelineEditor/Private/Graph/MovieGraphSchema.h
--- a/Engine/Plugins/MovieScene/MovieRenderPipeline/Source/MovieRenderPipelineEditor/Private/Graph/MovieGraphSchema.h
+++ b/Engine/Plugins/MovieScene/MovieRenderPipeline/Source/MovieRenderPipelineEditor/Private/Graph/MovieGraphSchema.h
@@ -31,6 +31,7 @@ public:
 	virtual UEdGraphPin* DropPinOnNode(UEdGraphNode* InTargetNode, const FName& InSourcePinName, const FEdGraphPinType& InSourcePinType, EEdGraphPinDirection InSourcePinDirection) const override;
 	virtual void SetPinBeingDroppedOnNode(UEdGraphPin* InSourcePin) const override;
 	virtual bool SafeDeleteNodeFromGraph(UEdGraph* Graph, UEdGraphNode* Node) const override;
+	virtual void OnPinConnectionDoubleCicked(UEdGraphPin* PinA, UEdGraphPin* PinB, const FVector2f& GraphPosition) const override;
 	//~ End EdGraphSchema Interface
 	
 	static void InitMoviePipelineNodeClasses();
