CIS: Fixed clang deprecated variable warnings.

diff --git a/Engine/Plugins/MovieScene/MovieRenderPipeline/Source/MovieRenderPipelineRenderPasses/Public/Graph/Renderers/MovieGraphImagePassBase.h b/Engine/Plugins/MovieScene/MovieRenderPipeline/Source/MovieRenderPipelineRenderPasses/Public/Graph/Renderers/MovieGraphImagePassBase.h
--- a/Engine/Plugins/MovieScene/MovieRenderPipeline/Source/MovieRenderPipelineRenderPasses/Public/Graph/Renderers/MovieGraphImagePassBase.h
+++ b/Engine/Plugins/MovieScene/MovieRenderPipeline/Source/MovieRenderPipelineRenderPasses/Public/Graph/Renderers/MovieGraphImagePassBase.h
@@ -21,6 +21,16 @@ namespace UE::MovieGraph::Rendering
 		bool bIsFirstSample;
 		UE_DEPRECATED(5.6, "No longer used, calling code now looks at sample state directly.")
 		bool bIsLastSample;
+
+
+		PRAGMA_DISABLE_DEPRECATION_WARNINGS
+		FMovieGraphRenderDataAccumulationArgs() = default;
+		~FMovieGraphRenderDataAccumulationArgs() = default;
+		FMovieGraphRenderDataAccumulationArgs(const FMovieGraphRenderDataAccumulationArgs&) = default;
+		FMovieGraphRenderDataAccumulationArgs(FMovieGraphRenderDataAccumulationArgs&&) = default;
+		FMovieGraphRenderDataAccumulationArgs& operator=(const FMovieGraphRenderDataAccumulationArgs&) = default;
+		FMovieGraphRenderDataAccumulationArgs& operator=(FMovieGraphRenderDataAccumulationArgs&&) = default;
+		PRAGMA_ENABLE_DEPRECATION_WARNINGS
 	};
 
 	// Forward Declare
