MoviePipeline: Default H.264 audio to "on" in MRG.

Also prevented an audio warning from showing up from the MP4 and ProRes nodes/settings if movies are being generated and the editor was started with "-nosound".

diff --git a/Engine/Plugins/Media/AppleProResMedia/Source/AppleProResMedia/Private/MovieGraphAppleProResNode.cpp b/Engine/Plugins/Media/AppleProResMedia/Source/AppleProResMedia/Private/MovieGraphAppleProResNode.cpp
--- a/Engine/Plugins/Media/AppleProResMedia/Source/AppleProResMedia/Private/MovieGraphAppleProResNode.cpp
+++ b/Engine/Plugins/Media/AppleProResMedia/Source/AppleProResMedia/Private/MovieGraphAppleProResNode.cpp
@@ -205,6 +205,13 @@ void UMovieGraphAppleProResNode::BeginFinalize_EncodeThread(MovieRenderGraph::IV
 		return;
 	}
 
+	// Nothing to do here if audio isn't being generated. The "invalid shot index" warning below is legitimate *if audio is being rendered*, but if
+	// no audio is being rendered (eg, with -nosound) then we don't want the warning to show up.
+	if (!FApp::CanEverRenderAudio())
+	{
+		return;
+	}
+
 	for (const TPair<int32, MovieRenderGraph::IVideoCodecWriter::FLightweightSourceData>& SourceData : CodecWriter->LightweightSourceData)
 	{
 		if (!AudioData.FinishedSegments.IsValidIndex(SourceData.Key))
diff --git a/Engine/Plugins/Media/AppleProResMedia/Source/AppleProResMedia/Private/MoviePipelineAppleProResOutput.cpp b/Engine/Plugins/Media/AppleProResMedia/Source/AppleProResMedia/Private/MoviePipelineAppleProResOutput.cpp
--- a/Engine/Plugins/Media/AppleProResMedia/Source/AppleProResMedia/Private/MoviePipelineAppleProResOutput.cpp
+++ b/Engine/Plugins/Media/AppleProResMedia/Source/AppleProResMedia/Private/MoviePipelineAppleProResOutput.cpp
@@ -117,6 +117,13 @@ void UMoviePipelineAppleProResOutput::BeginFinalize_EncodeThread(MovieRenderPipe
 	{
 		return;
 	}
+	
+	// Nothing to do here if audio isn't being generated. The "invalid shot index" warning below is legitimate *if audio is being rendered*, but if
+	// no audio is being rendered (eg, with -nosound) then we don't want the warning to show up.
+	if (!FApp::CanEverRenderAudio())
+	{
+		return;
+	}
 
 	for (const TPair<int32, MovieRenderPipeline::IVideoCodecWriter::FLightweightSourceData>& SourceData : CodecWriter->LightweightSourceData)
 	{
diff --git a/Engine/Plugins/MovieScene/MovieRenderPipeline/Source/MovieRenderPipelineMP4Encoder/Private/Graph/MovieGraphMP4EncoderNode.cpp b/Engine/Plugins/MovieScene/MovieRenderPipeline/Source/MovieRenderPipelineMP4Encoder/Private/Graph/MovieGraphMP4EncoderNode.cpp
--- a/Engine/Plugins/MovieScene/MovieRenderPipeline/Source/MovieRenderPipelineMP4Encoder/Private/Graph/MovieGraphMP4EncoderNode.cpp
+++ b/Engine/Plugins/MovieScene/MovieRenderPipeline/Source/MovieRenderPipelineMP4Encoder/Private/Graph/MovieGraphMP4EncoderNode.cpp
@@ -180,6 +180,13 @@ void UMovieGraphMP4EncoderNode::BeginFinalize_EncodeThread(MovieRenderGraph::IVi
 	{
 		return;
 	}
+	
+	// Nothing to do here if audio isn't being generated. The "invalid shot index" warning below is legitimate *if audio is being rendered*, but if
+	// no audio is being rendered (eg, with -nosound) then we don't want the warning to show up.
+	if (!FApp::CanEverRenderAudio())
+	{
+		return;
+	}
 
 	const MoviePipeline::FAudioState& AudioData = CachedPipeline->GetAudioRendererInstance()->GetAudioState();
 
diff --git a/Engine/Plugins/MovieScene/MovieRenderPipeline/Source/MovieRenderPipelineMP4Encoder/Private/MoviePipelineMP4EncoderOutput.cpp b/Engine/Plugins/MovieScene/MovieRenderPipeline/Source/MovieRenderPipelineMP4Encoder/Private/MoviePipelineMP4EncoderOutput.cpp
--- a/Engine/Plugins/MovieScene/MovieRenderPipeline/Source/MovieRenderPipelineMP4Encoder/Private/MoviePipelineMP4EncoderOutput.cpp
+++ b/Engine/Plugins/MovieScene/MovieRenderPipeline/Source/MovieRenderPipelineMP4Encoder/Private/MoviePipelineMP4EncoderOutput.cpp
@@ -87,6 +87,13 @@ void UMoviePipelineMP4EncoderOutput::BeginFinalize_EncodeThread(MovieRenderPipel
 		return;
 	}
 
+	// Nothing to do here if audio isn't being generated. The "invalid shot index" warning below is legitimate *if audio is being rendered*, but if
+	// no audio is being rendered (eg, with -nosound) then we don't want the warning to show up.
+	if (!FApp::CanEverRenderAudio())
+	{
+		return;
+	}
+
 	for (const TPair<int32, MovieRenderPipeline::IVideoCodecWriter::FLightweightSourceData>& SourceData : CodecWriter->LightweightSourceData)
 	{
 		if (!AudioData.FinishedSegments.IsValidIndex(SourceData.Key))
diff --git a/Engine/Plugins/MovieScene/MovieRenderPipeline/Source/MovieRenderPipelineMP4Encoder/Public/Graph/MovieGraphMP4EncoderNode.h b/Engine/Plugins/MovieScene/MovieRenderPipeline/Source/MovieRenderPipelineMP4Encoder/Public/Graph/MovieGraphMP4EncoderNode.h
--- a/Engine/Plugins/MovieScene/MovieRenderPipeline/Source/MovieRenderPipelineMP4Encoder/Public/Graph/MovieGraphMP4EncoderNode.h
+++ b/Engine/Plugins/MovieScene/MovieRenderPipeline/Source/MovieRenderPipelineMP4Encoder/Public/Graph/MovieGraphMP4EncoderNode.h
@@ -88,7 +88,7 @@ public:
 
 	/** If true, audio will be included in the video file. */
 	UPROPERTY(EditAnywhere, BlueprintReadWrite, Category = "H.264", meta = (EditCondition = "bOverride_bIncludeAudio"))
-	bool bIncludeAudio;
+	bool bIncludeAudio = true;
 	
 	/**
 	* What is the average bitrate the encoder should target per second? Value is in Megabits per Second,
