MRQ: Fixed darkened pixel around edges of rendered frames.

This was happening when rendering overlapped tiles and constraining the aspect ratio, and at least 2 AA samples. The reason for this is that subpixel AA offsets will sample outside the active area around the edges when the tiles are bigger than the active frame, and because the letterbox is applied prior to accumulation with transparent black, that outer rim was bleeding into the active area.

As a solution, GPU letterboxing now leaves a 1 pixel wide margin around the active area so that the active area accumulation can happen correctly, and then after the last tile is accumulated, that extra margin is then overwritten with transparent black. This latter letterbox outline needs to be 2 pixels wide instead of 1, because the extra pixel that we left before will bleed 1 pixel further out during AA accumulation.

#rnx

diff --git a/Engine/Plugins/MovieScene/MovieRenderPipeline/Source/MovieRenderPipelineRenderPasses/Private/MoviePipelineDeferredPasses.cpp b/Engine/Plugins/MovieScene/MovieRenderPipeline/Source/MovieRenderPipelineRenderPasses/Private/MoviePipelineDeferredPasses.cpp
--- a/Engine/Plugins/MovieScene/MovieRenderPipeline/Source/MovieRenderPipelineRenderPasses/Private/MoviePipelineDeferredPasses.cpp
+++ b/Engine/Plugins/MovieScene/MovieRenderPipeline/Source/MovieRenderPipelineRenderPasses/Private/MoviePipelineDeferredPasses.cpp
@@ -1040,8 +1040,12 @@ void UMoviePipelineDeferredPassBase::PostRendererSubmission(const FMoviePipeline
 	{
 		return;
 	}
-	
-	// Draw letterboxing
+
+	// Draw letterboxing, except for 1 pixel outline for anti-aliasing purposes.
+	// This extra pixel thickness will be done later after accumulating the last tile.
+
+	MoviePipeline::FLetterboxData LetterboxData;
+
 	APlayerCameraManager* PlayerCameraManager = GetPipeline()->GetWorld()->GetFirstPlayerController()->PlayerCameraManager;
 	if(PlayerCameraManager && PlayerCameraManager->GetCameraCacheView().bConstrainAspectRatio)
 	{
@@ -1062,40 +1066,69 @@ void UMoviePipelineDeferredPassBase::PostRendererSubmission(const FMoviePipeline
 		const FIntPoint TileViewMin = InSampleState.OverlappedOffset;
 		const FIntPoint TileViewMax = TileViewMin + InSampleState.BackbufferSize;
 
-		// Camera ratio constrained rect, clipped by the tile rect
+		// Calculate camera ratio constrained rect.
+
 		FIntPoint ConstrainedViewMin = (FullOutputSize - ConstrainedFullSize) / 2;
 		FIntPoint ConstrainedViewMax = ConstrainedViewMin + ConstrainedFullSize;
+
+		// When there is overlap, adjust the offsets to leave a 1-pixel margin around the active pixel area, so that
+		// anti-aliasing is able to sample the original pixel during accumulation, instead of 
+		// the letterbox pixel which would otherwise darken the edges of the active area.
+		LetterboxData.bDrawLetterboxBorder = InSampleState.OverlappedPad != FIntPoint::ZeroValue;
+		LetterboxData.FrameActiveArea = FIntRect(ConstrainedViewMin, ConstrainedViewMax);
+
+		// Clip it by the tile rect.
+
 		ConstrainedViewMin = FIntPoint(FMath::Clamp(ConstrainedViewMin.X, TileViewMin.X, TileViewMax.X),
 			FMath::Clamp(ConstrainedViewMin.Y, TileViewMin.Y, TileViewMax.Y));
+
 		ConstrainedViewMax = FIntPoint(FMath::Clamp(ConstrainedViewMax.X, TileViewMin.X, TileViewMax.X),
 			FMath::Clamp(ConstrainedViewMax.Y, TileViewMin.Y, TileViewMax.Y));
 
 		// Difference between the clipped constrained rect and the tile rect
+
 		const FIntPoint OffsetMin = ConstrainedViewMin - TileViewMin;
 		const FIntPoint OffsetMax = TileViewMax - ConstrainedViewMax;
 
+		const int32 AntiAliasMargin = LetterboxData.bDrawLetterboxBorder ? 1 : 0;
+
+		const int32 LeftThickness   = FMath::Max(OffsetMin.X - AntiAliasMargin, 0);
+		const int32 RightThickness  = FMath::Max(OffsetMax.X - AntiAliasMargin, 0);
+		const int32 TopThickness    = FMath::Max(OffsetMin.Y - AntiAliasMargin, 0);
+		const int32 BottomThickness = FMath::Max(OffsetMax.Y - AntiAliasMargin, 0);
+
+		// These thicknesses will inform a later stage which pixels have been cleared out and that 
+		// should not be sampled.
+		LetterboxData.LeftSamplePixelsClearedBeforeAccumulation = LeftThickness;
+		LetterboxData.RightSamplePixelsClearedBeforeAccumulation = RightThickness;
+		LetterboxData.TopSamplePixelsClearedBeforeAccumulation = TopThickness;
+		LetterboxData.BottomSamplePixelsClearedBeforeAccumulation = BottomThickness;
+
 		// Clear left
-		if (OffsetMin.X > 0)
+		if (LeftThickness > 0)
 		{
-			InCanvas.DrawTile(0, 0, OffsetMin.X, InSampleState.BackbufferSize.Y,
+			InCanvas.DrawTile(0, 0, LeftThickness, InSampleState.BackbufferSize.Y,
 				0.0f, 0.0f, 1.0f, 1.0f, FLinearColor::Black, nullptr, false);
 		}
+
 		// Clear right
-		if (OffsetMax.X > 0)
+		if (RightThickness > 0)
 		{
-			InCanvas.DrawTile(InSampleState.BackbufferSize.X - OffsetMax.X, 0, InSampleState.BackbufferSize.X, InSampleState.BackbufferSize.Y,
+			InCanvas.DrawTile(InSampleState.BackbufferSize.X - RightThickness, 0, InSampleState.BackbufferSize.X, InSampleState.BackbufferSize.Y,
 				0.0f, 0.0f, 1.0f, 1.0f, FLinearColor::Black, nullptr, false);
 		}
+
 		// Clear top
-		if (OffsetMin.Y > 0)
+		if (TopThickness > 0)
 		{
-			InCanvas.DrawTile(0, 0, InSampleState.BackbufferSize.X, OffsetMin.Y,
+			InCanvas.DrawTile(0, 0, InSampleState.BackbufferSize.X, TopThickness,
 				0.0f, 0.0f, 1.0f, 1.0f, FLinearColor::Black, nullptr, false);
 		}
+
 		// Clear bottom
-		if (OffsetMax.Y > 0)
+		if (BottomThickness > 0)
 		{
-			InCanvas.DrawTile(0, InSampleState.BackbufferSize.Y - OffsetMax.Y, InSampleState.BackbufferSize.X, InSampleState.BackbufferSize.Y,
+			InCanvas.DrawTile(0, InSampleState.BackbufferSize.Y - BottomThickness, InSampleState.BackbufferSize.X, InSampleState.BackbufferSize.Y,
 				0.0f, 0.0f, 1.0f, 1.0f, FLinearColor::Black, nullptr, false);
 		}
 
@@ -1124,6 +1157,7 @@ void UMoviePipelineDeferredPassBase::PostRendererSubmission(const FMoviePipeline
 		AccumulationArgs.OutputMerger = GetPipeline()->OutputBuilder;
 		AccumulationArgs.ImageAccumulator = StaticCastSharedPtr<FImageOverlappedAccumulator>(SampleAccumulator->Accumulator);
 		AccumulationArgs.bAccumulateAlpha = bAccumulatorIncludesAlpha;
+		AccumulationArgs.LetterboxData = LetterboxData;
 	}
 
 	auto Callback = [this, FramePayload, AccumulationArgs, SampleAccumulator](TUniquePtr<FImagePixelData>&& InPixelData)
@@ -1139,6 +1173,7 @@ void UMoviePipelineDeferredPassBase::PostRendererSubmission(const FMoviePipeline
 		{
 			// Enqueue a encode for this frame onto our worker thread.
 			MoviePipeline::AccumulateSample_TaskThread(MoveTemp(PixelData), AccumulationArgs);
+
 			if (bFinalSample)
 			{
 				// Final sample has now been executed, break the pre-req chain and free the accumulator for reuse.
diff --git a/Engine/Plugins/MovieScene/MovieRenderPipeline/Source/MovieRenderPipelineRenderPasses/Private/MoviePipelineImagePassBase.cpp b/Engine/Plugins/MovieScene/MovieRenderPipeline/Source/MovieRenderPipelineRenderPasses/Private/MoviePipelineImagePassBase.cpp
--- a/Engine/Plugins/MovieScene/MovieRenderPipeline/Source/MovieRenderPipelineRenderPasses/Private/MoviePipelineImagePassBase.cpp
+++ b/Engine/Plugins/MovieScene/MovieRenderPipeline/Source/MovieRenderPipelineRenderPasses/Private/MoviePipelineImagePassBase.cpp
@@ -946,6 +946,115 @@ void FAccumulatorPool::FAccumulatorInstance::SetIsActive(const bool bInIsActive)
 
 namespace MoviePipeline
 {
+	/**
+	 * Clears the letterbox border that was not already cleared in GPU.
+	 * Note: It was left this way for proper anti-aliasing at the edges of the frame.
+	 * 
+	 * @param LetterboxData - Data about the border, including whether it is enabled or not.
+	 * @param ImageData     - Pixel data to draw on.
+	 */
+	void DrawLetterboxBorder(const FLetterboxData& LetterboxData, FImagePixelData* ImageData)
+	{
+		if (!ImageData || !LetterboxData.bDrawLetterboxBorder)
+		{
+			return;
+		}
+
+		TRACE_CPUPROFILER_EVENT_SCOPE(MoviePipeline::DrawLetterboxBorder);
+
+		constexpr int32 BorderThickness = 2;
+		const FIntRect& FrameActiveArea = LetterboxData.FrameActiveArea;
+
+		// Get the overall image dimensions.
+		const FIntPoint ImageSize = ImageData->GetSize();
+		const int32 FullWidth = ImageSize.X;
+		const int32 FullHeight = ImageSize.Y;
+
+		// Generic lambda to clear a rectangular region within a pixel array.
+		auto ClearRegion = [FullWidth](auto& Pixels, int32 X0, int32 X1, int32 Y0, int32 Y1)
+			{
+				if (X0 >= X1 || Y0 >= Y1)
+				{
+					return;
+				}
+
+				for (int32 Y = Y0; Y < Y1; ++Y)
+				{
+					for (int32 X = X0; X < X1; ++X)
+					{
+						Pixels[Y * FullWidth + X] = {}; // Transparent black
+					}
+				}
+			};
+
+		// Lambda to draw all four borders using the ClearRegion helper.
+		auto DrawBorders = [&](auto& Pixels)
+			{
+				// Top border (includes top-left and top-right corners)
+				{
+					const int32 X0 = FMath::Max(FrameActiveArea.Min.X - BorderThickness, 0);
+					const int32 X1 = FMath::Min(FrameActiveArea.Max.X + BorderThickness, FullWidth);
+					const int32 Y0 = FMath::Max(FrameActiveArea.Min.Y - BorderThickness, 0);
+					const int32 Y1 = FrameActiveArea.Min.Y;
+
+					ClearRegion(Pixels, X0, X1, Y0, Y1);
+				}
+
+				// Bottom border (includes bottom-left and bottom-right corners)
+				{
+					const int32 X0 = FMath::Max(FrameActiveArea.Min.X - BorderThickness, 0);
+					const int32 X1 = FMath::Min(FrameActiveArea.Max.X + BorderThickness, FullWidth);
+					const int32 Y0 = FrameActiveArea.Max.Y;
+					const int32 Y1 = FMath::Min(FrameActiveArea.Max.Y + BorderThickness, FullHeight);
+
+					ClearRegion(Pixels, X0, X1, Y0, Y1);
+				}
+
+				// Left border
+				{
+					const int32 X0 = FMath::Max(FrameActiveArea.Min.X - BorderThickness, 0);
+					const int32 X1 = FrameActiveArea.Min.X;
+					const int32 Y0 = FrameActiveArea.Min.Y;
+					const int32 Y1 = FrameActiveArea.Max.Y;
+
+					ClearRegion(Pixels, X0, X1, Y0, Y1);
+				}
+
+				// Right border
+				{
+					const int32 X0 = FrameActiveArea.Max.X;
+					const int32 X1 = FMath::Min(FrameActiveArea.Max.X + BorderThickness, FullWidth);
+					const int32 Y0 = FrameActiveArea.Min.Y;
+					const int32 Y1 = FrameActiveArea.Max.Y;
+
+					ClearRegion(Pixels, X0, X1, Y0, Y1);
+				}
+			};
+
+		// Dispatch based on the pixel type.
+		switch (ImageData->GetType())
+		{
+		case EImagePixelType::Color:
+		{
+			DrawBorders(static_cast<TImagePixelData<FColor>*>(ImageData)->Pixels);
+			break;
+		}
+		case EImagePixelType::Float16:
+		{
+			DrawBorders(static_cast<TImagePixelData<FFloat16Color>*>(ImageData)->Pixels);
+			break;
+		}
+		case EImagePixelType::Float32:
+		{
+			DrawBorders(static_cast<TImagePixelData<FLinearColor>*>(ImageData)->Pixels);
+			break;
+		}
+		default:
+			checkNoEntry();
+			break;
+		}
+	}
+
 	void AccumulateSample_TaskThread(TUniquePtr<FImagePixelData>&& InPixelData, const MoviePipeline::FImageSampleAccumulationArgs& InParams)
 	{
 		SCOPE_CYCLE_COUNTER(STAT_AccumulateSample_TaskThread);
@@ -991,18 +1100,25 @@ namespace MoviePipeline
 			InParams.OutputMerger.Pin()->OnSingleSampleDataAvailable_AnyThread(MoveTemp(SampleData));
 		}
 
-		// Optimization! If we don't need the accumulator (no tiling, no supersampling) then we'll skip it and just send it straight to the output stage.
-		// This significantly improves performance in the baseline case.
-		const bool bOneTile = NewPayload->IsFirstTile() && NewPayload->IsLastTile();
-		const bool bOneTS = NewPayload->IsFirstTemporalSample() && NewPayload->IsLastTemporalSample();
-		const bool bOneSS = NewPayload->SampleState.SpatialSampleCount == 1;
+		const bool bHasOverlap = NewPayload->SampleState.OverlappedPad != FIntPoint::ZeroValue;
 
-		if (bOneTile && bOneTS && bOneSS)
+		// Optimization! If we don't need the accumulator (no tiling, no supersampling, no overlap) then we'll skip it 
+		// and just send it straight to the output stage, significantly improving performance in the baseline case.
 		{
-			// Send the data directly to the Output Builder and skip the accumulator.
-			ensure(InParams.OutputMerger.IsValid());
-			InParams.OutputMerger.Pin()->OnCompleteRenderPassDataAvailable_AnyThread(MoveTemp(SamplePixelData));
-			return;
+			const bool bOneTile = NewPayload->IsFirstTile() && NewPayload->IsLastTile();
+			const bool bOneTS = NewPayload->IsFirstTemporalSample() && NewPayload->IsLastTemporalSample();
+			const bool bOneSS = NewPayload->SampleState.SpatialSampleCount == 1;
+
+			if (bOneTile && bOneTS && bOneSS && !bHasOverlap)
+			{
+				// We do not expect deferred letterbox drawing without tile overlap present.
+				check(!InParams.LetterboxData.bDrawLetterboxBorder);
+
+				// Send the data directly to the Output Builder and skip the accumulator.
+				ensure(InParams.OutputMerger.IsValid());
+				InParams.OutputMerger.Pin()->OnCompleteRenderPassDataAvailable_AnyThread(MoveTemp(SamplePixelData));
+				return;
+			}
 		}
 
 		// Allocate memory if the ImageAccumulator has not been initialized yet for this output
@@ -1011,6 +1127,26 @@ namespace MoviePipeline
 		MoviePipeline::FTileWeight1D WeightFunctionY;
 		NewPayload->GetWeightFunctionParams(/*Out*/ WeightFunctionX, /*Out*/ WeightFunctionY);
 
+		// Adjust the weights to account for the pixels that were cleared before accumulation,
+		// and should therefore not be sampled.
+		// 
+		// Note: We exclude overlap cases which should have the anti-aliasing margin with real pixels
+		// already and do not really need this sampling protection. Doing so is slightly more complicated
+		// because they will have the finite slopes in the _/-\_ weights 1D curve and would probably need
+		// to add MinX and MaxX limit notions to FTileWeight1D and use that instead to keep the slopes intact.
+		if (!bHasOverlap)
+		{
+			WeightFunctionX.X0 = FMath::Max(WeightFunctionX.X0, InParams.LetterboxData.LeftSamplePixelsClearedBeforeAccumulation);
+			WeightFunctionX.X1 = FMath::Max(WeightFunctionX.X1, InParams.LetterboxData.LeftSamplePixelsClearedBeforeAccumulation);
+			WeightFunctionX.X2 = FMath::Min(WeightFunctionX.X2, SamplePixelData->GetSize().X - InParams.LetterboxData.RightSamplePixelsClearedBeforeAccumulation);
+			WeightFunctionX.X3 = FMath::Min(WeightFunctionX.X3, SamplePixelData->GetSize().X - InParams.LetterboxData.RightSamplePixelsClearedBeforeAccumulation);
+
+			WeightFunctionY.X0 = FMath::Max(WeightFunctionY.X0, InParams.LetterboxData.TopSamplePixelsClearedBeforeAccumulation);
+			WeightFunctionY.X1 = FMath::Max(WeightFunctionY.X1, InParams.LetterboxData.TopSamplePixelsClearedBeforeAccumulation);
+			WeightFunctionY.X2 = FMath::Min(WeightFunctionY.X2, SamplePixelData->GetSize().Y - InParams.LetterboxData.BottomSamplePixelsClearedBeforeAccumulation);
+			WeightFunctionY.X3 = FMath::Min(WeightFunctionY.X3, SamplePixelData->GetSize().Y - InParams.LetterboxData.BottomSamplePixelsClearedBeforeAccumulation);
+		}
+
 		TSharedPtr<FImageOverlappedAccumulator> PinnedImageAccumulator = InParams.ImageAccumulator.Pin();
 		TSharedPtr<IMoviePipelineOutputMerger> PinnedOutputMerger = InParams.OutputMerger.Pin();
 
@@ -1132,6 +1268,9 @@ namespace MoviePipeline
 				TUniquePtr<TImagePixelData<FLinearColor> > FinalPixelData = MakeUnique<TImagePixelData<FLinearColor>>(FIntPoint(FullSizeX, FullSizeY), NewPayload);
 				PinnedImageAccumulator->FetchFinalPixelDataLinearColor(FinalPixelData->Pixels);
 
+				// Apply letterbox outline. Will only do any work if enabled.
+				DrawLetterboxBorder(InParams.LetterboxData, FinalPixelData.Get());
+
 				// Send the data to the Output Builder
 				PinnedOutputMerger->OnCompleteRenderPassDataAvailable_AnyThread(MoveTemp(FinalPixelData));
 			}
@@ -1141,6 +1280,9 @@ namespace MoviePipeline
 				TUniquePtr<TImagePixelData<FFloat16Color> > FinalPixelData = MakeUnique<TImagePixelData<FFloat16Color>>(FIntPoint(FullSizeX, FullSizeY), NewPayload);
 				PinnedImageAccumulator->FetchFinalPixelDataHalfFloat(FinalPixelData->Pixels);
 
+				// Apply letterbox outline. Will only do any work if enabled.
+				DrawLetterboxBorder(InParams.LetterboxData, FinalPixelData.Get());
+
 				// Send the data to the Output Builder
 				PinnedOutputMerger->OnCompleteRenderPassDataAvailable_AnyThread(MoveTemp(FinalPixelData));
 			}
@@ -1150,6 +1292,9 @@ namespace MoviePipeline
 				TUniquePtr<TImagePixelData<FColor>> FinalPixelData = MakeUnique<TImagePixelData<FColor>>(FIntPoint(FullSizeX, FullSizeY), NewPayload);
 				PinnedImageAccumulator->FetchFinalPixelDataByte(FinalPixelData->Pixels);
 
+				// Apply letterbox outline. Will only do any work if enabled.
+				DrawLetterboxBorder(InParams.LetterboxData, FinalPixelData.Get());
+
 				// Send the data to the Output Builder
 				PinnedOutputMerger->OnCompleteRenderPassDataAvailable_AnyThread(MoveTemp(FinalPixelData));
 			}
diff --git a/Engine/Plugins/MovieScene/MovieRenderPipeline/Source/MovieRenderPipelineRenderPasses/Public/MoviePipelineImagePassBase.h b/Engine/Plugins/MovieScene/MovieRenderPipeline/Source/MovieRenderPipelineRenderPasses/Public/MoviePipelineImagePassBase.h
--- a/Engine/Plugins/MovieScene/MovieRenderPipeline/Source/MovieRenderPipelineRenderPasses/Public/MoviePipelineImagePassBase.h
+++ b/Engine/Plugins/MovieScene/MovieRenderPipeline/Source/MovieRenderPipelineRenderPasses/Public/MoviePipelineImagePassBase.h
@@ -72,12 +72,45 @@ public:
 
 namespace MoviePipeline
 {
+	/** Describes the data required to clear the letterbox pixels outside the active pixel area */
+	struct FLetterboxData
+	{
+		/** Active pixel area of the full frame within channel plane dimensions. i.e. Region not letterboxed */
+		FIntRect FrameActiveArea;
+
+		/** True if letterboxing is active and drawing the border was deferred */
+		bool bDrawLetterboxBorder;
+
+		/** Left sample pixels already cleared that should therefore not be sampled */
+		int32 LeftSamplePixelsClearedBeforeAccumulation;
+
+		/** Right sample pixels already cleared that should therefore not be sampled */
+		int32 RightSamplePixelsClearedBeforeAccumulation;
+
+		/** Top sample pixels already cleared that should therefore not be sampled */
+		int32 TopSamplePixelsClearedBeforeAccumulation;
+
+		/** Bottom sample pixels already cleared that should therefore not be sampled */
+		int32 BottomSamplePixelsClearedBeforeAccumulation;
+
+		FLetterboxData() :
+			  bDrawLetterboxBorder(false)
+			, LeftSamplePixelsClearedBeforeAccumulation(0)
+			, RightSamplePixelsClearedBeforeAccumulation(0)
+			, TopSamplePixelsClearedBeforeAccumulation(0)
+			, BottomSamplePixelsClearedBeforeAccumulation(0)
+		{};
+	};
+
 	struct FImageSampleAccumulationArgs
 	{
 	public:
 		TWeakPtr<FImageOverlappedAccumulator, ESPMode::ThreadSafe> ImageAccumulator;
 		TWeakPtr<IMoviePipelineOutputMerger, ESPMode::ThreadSafe> OutputMerger;
 		bool bAccumulateAlpha;
+
+		/** Contains the data needed for clearing the letterbox area */
+		FLetterboxData LetterboxData;
 	};
 
 	void MOVIERENDERPIPELINERENDERPASSES_API AccumulateSample_TaskThread(TUniquePtr<FImagePixelData>&& InPixelData, const MoviePipeline::FImageSampleAccumulationArgs& InParams);
