MoviePipeline: Fixed issue where media would not be written to disk when a UNC path was used.

diff --git a/Engine/Plugins/MovieScene/MovieRenderPipeline/Source/MovieRenderPipelineCore/Private/Graph/MovieGraphBlueprintLibrary.cpp b/Engine/Plugins/MovieScene/MovieRenderPipeline/Source/MovieRenderPipelineCore/Private/Graph/MovieGraphBlueprintLibrary.cpp
--- a/Engine/Plugins/MovieScene/MovieRenderPipeline/Source/MovieRenderPipelineCore/Private/Graph/MovieGraphBlueprintLibrary.cpp
+++ b/Engine/Plugins/MovieScene/MovieRenderPipeline/Source/MovieRenderPipelineCore/Private/Graph/MovieGraphBlueprintLibrary.cpp
@@ -193,8 +193,22 @@ FString UMovieGraphBlueprintLibrary::ResolveFilenameFormatArguments(const FStrin
 	}
 
 	// Fix-up slashes
-	FPaths::NormalizeFilename(BaseFilename);
-	FPaths::RemoveDuplicateSlashes(BaseFilename);
+	{
+		FPaths::NormalizeFilename(BaseFilename);
+
+		// Remove duplicate slashes. Recognize a double-slash at the beginning to ensure network paths remain intact. Note that this also works on
+		// network paths that start with "\\" because NormalizeFilename() will convert "\\" to "//".
+		if (BaseFilename.StartsWith(TEXT("//")))
+		{
+			BaseFilename.RightChopInline(1, EAllowShrinking::No);
+			FPaths::RemoveDuplicateSlashes(BaseFilename);
+			BaseFilename = TEXT("/") + MoveTemp(BaseFilename);
+		}
+		else
+		{
+			FPaths::RemoveDuplicateSlashes(BaseFilename);
+		}
+	}
 
 	// In the event of multiple dots in the filename, we need to replace them with a single dot to ensure a valid file name
 	while (BaseFilename.Contains(".."))
