MoviePipeline: Added a new setting to the Global Game Overrides node ("Rebuild Lumen Scene Between Render Layers") that can help with Lumen consistency between render layers.

diff --git a/Engine/Plugins/MovieScene/MovieRenderPipeline/Source/MovieRenderPipelineCore/Private/Graph/Nodes/MovieGraphGlobalGameOverrides.cpp b/Engine/Plugins/MovieScene/MovieRenderPipeline/Source/MovieRenderPipelineCore/Private/Graph/Nodes/MovieGraphGlobalGameOverrides.cpp
--- a/Engine/Plugins/MovieScene/MovieRenderPipeline/Source/MovieRenderPipelineCore/Private/Graph/Nodes/MovieGraphGlobalGameOverrides.cpp
+++ b/Engine/Plugins/MovieScene/MovieRenderPipeline/Source/MovieRenderPipelineCore/Private/Graph/Nodes/MovieGraphGlobalGameOverrides.cpp
@@ -21,6 +21,7 @@ PRAGMA_ENABLE_DEPRECATION_WARNINGS
 	, bFlushGrassStreaming(false)
 	, bFlushStreamingManagers(false)
 	, VirtualTextureFeedbackFactor(1)
+	, bRebuildLumenSceneBetweenRenderLayers(false)
 {
 }
 
@@ -71,6 +72,14 @@ void UMovieGraphGlobalGameOverridesNode::BuildNewProcessCommandLineArgsImpl(TArr
 		InOutDeviceProfileCvars.AddUnique(TEXT("r.Streaming.SyncStatesWhenBlocking=1"));
 	}
 
+	if (bRebuildLumenSceneBetweenRenderLayers)
+	{
+		InOutDeviceProfileCvars.AddUnique(TEXT("r.LumenScene.Radiosity.UpdateFactor=1"));
+		InOutDeviceProfileCvars.AddUnique(TEXT("r.LumenScene.SurfaceCache.CardCaptureFactor=1"));
+		InOutDeviceProfileCvars.AddUnique(TEXT("r.LumenScene.SurfaceCache.Feedback=0"));
+		InOutDeviceProfileCvars.AddUnique(TEXT("r.LumenScene.SurfaceCache.RecaptureEveryFrame=1"));
+	}
+
 	// Like the extra cvars applied in ApplySettings(), the below are applied to allow MRQ to function correctly.
 
 #if WITH_EDITOR
@@ -173,6 +182,14 @@ void UMovieGraphGlobalGameOverridesNode::ApplySettings(const bool bOverrideValue
 		MOVIEPIPELINE_STORE_AND_OVERRIDE_CVAR_INT(NodeCDO->PreviousStreamingManagerSyncState, TEXT("r.Streaming.SyncStatesWhenBlocking"), 1, bOverrideValues);
 	}
 
+	if (bRebuildLumenSceneBetweenRenderLayers)
+	{
+		MOVIEPIPELINE_STORE_AND_OVERRIDE_CVAR_INT(NodeCDO->PreviousLumenRadiosityUpdateFactor, TEXT("r.LumenScene.Radiosity.UpdateFactor"), 1, bOverrideValues);
+		MOVIEPIPELINE_STORE_AND_OVERRIDE_CVAR_INT(NodeCDO->PreviousLumenSurfaceCacheCardCaptureFactor, TEXT("r.LumenScene.SurfaceCache.CardCaptureFactor"), 1, bOverrideValues);
+		MOVIEPIPELINE_STORE_AND_OVERRIDE_CVAR_INT(NodeCDO->PreviousLumenSurfaceCacheFeedback, TEXT("r.LumenScene.SurfaceCache.Feedback"), 0, bOverrideValues);
+		MOVIEPIPELINE_STORE_AND_OVERRIDE_CVAR_INT(NodeCDO->PreviousLumenSurfaceCacheRecaptureEveryFrame, TEXT("r.LumenScene.SurfaceCache.RecaptureEveryFrame"), 1, bOverrideValues);
+	}
+
 	MOVIEPIPELINE_STORE_AND_OVERRIDE_CVAR_BOOL(PreviousAlphaOutput, TEXT("r.PostProcessing.PropagateAlpha"), UE::MoviePipeline::GetAlphaOutputOverride(), bOverrideValues);
 
 	// Set cvars that allow MRQ to run correctly. These are not exposed as properties on the node, but the chance of a
diff --git a/Engine/Plugins/MovieScene/MovieRenderPipeline/Source/MovieRenderPipelineCore/Public/Graph/Nodes/MovieGraphGlobalGameOverrides.h b/Engine/Plugins/MovieScene/MovieRenderPipeline/Source/MovieRenderPipelineCore/Public/Graph/Nodes/MovieGraphGlobalGameOverrides.h
--- a/Engine/Plugins/MovieScene/MovieRenderPipeline/Source/MovieRenderPipelineCore/Public/Graph/Nodes/MovieGraphGlobalGameOverrides.h
+++ b/Engine/Plugins/MovieScene/MovieRenderPipeline/Source/MovieRenderPipelineCore/Public/Graph/Nodes/MovieGraphGlobalGameOverrides.h
@@ -97,6 +97,9 @@ public:
 	UPROPERTY(EditAnywhere, BlueprintReadWrite, Category = Overrides, meta = (InlineEditConditionToggle))
 	uint8 bOverride_VirtualTextureFeedbackFactor : 1;
 
+	UPROPERTY(EditAnywhere, BlueprintReadWrite, Category = Overrides, meta = (InlineEditConditionToggle))
+	uint8 bOverride_bRebuildLumenSceneBetweenRenderLayers : 1;
+
 	/**
 	 * Optional game mode to override the map's default game mode with. This can be useful if the game's normal mode
 	 * displays UI elements or loading screens that you don't want captured.
@@ -201,6 +204,22 @@ public:
 	UPROPERTY(EditAnywhere, BlueprintReadWrite, Category = "Rendering", meta = (EditCondition = "bOverride_VirtualTextureFeedbackFactor"))
 	int32 VirtualTextureFeedbackFactor;
 
+	/**
+	 * When rendering multiple Render Layers, Lumen's surface cache may not be able to keep up with content switching between hidden/visible
+	 * states. This setting mitigates the problem by resetting the surface cache for each Render Layer on each frame. Enabling this setting
+	 * has a performance cost. Disabling Screen Traces can also improve consistency when using Render Layers with Lumen. This setting may be
+	 * removed in a future engine release if it is no longer needed.
+	 *
+	 * Configures the following Console Variables:
+	 *
+	 * - r.LumenScene.Radiosity.UpdateFactor = 1
+	 * - r.LumenScene.SurfaceCache.CardCaptureFactor = 1
+	 * - r.LumenScene.SurfaceCache.Feedback = 0
+	 * - r.LumenScene.SurfaceCache.RecaptureEveryFrame = 1
+	 */
+	UPROPERTY(EditAnywhere, BlueprintReadWrite, Category = "Lumen Render Layers (Experimental)", meta = (EditCondition = "bOverride_bRebuildLumenSceneBetweenRenderLayers"))
+	bool bRebuildLumenSceneBetweenRenderLayers;
+
 private:
 	Scalability::FQualityLevels PreviousQualityLevels;
 	int32 PreviousTextureStreaming;
@@ -210,6 +229,10 @@ private:
 	int32 PreviousFoliageDitheredLOD;
 	int32 PreviousFoliageForceLOD;
 	int32 PreviousStreamingManagerSyncState;
+	int32 PreviousLumenRadiosityUpdateFactor;
+	int32 PreviousLumenSurfaceCacheCardCaptureFactor;
+	int32 PreviousLumenSurfaceCacheFeedback;
+	int32 PreviousLumenSurfaceCacheRecaptureEveryFrame;
 	int32 PreviousAnimationUROEnabled;
 	int32 PreviousSkyLightRealTimeReflectionCaptureTimeSlice;
 	int32 PreviousVolumetricRenderTarget;
