MovieGraph: Fixed "Resize PIEWindow to Output Resolution" not working when using the Movie Graph. Only the basic output resolution setting is supported, even though MRG can ultimately support multiple output resolutions at once (due to multi-camera support and adaptive resolutions).

diff --git a/Engine/Plugins/MovieScene/MovieRenderPipeline/Source/MovieRenderPipelineEditor/Private/MoviePipelinePIEExecutor.cpp b/Engine/Plugins/MovieScene/MovieRenderPipeline/Source/MovieRenderPipelineEditor/Private/MoviePipelinePIEExecutor.cpp
--- a/Engine/Plugins/MovieScene/MovieRenderPipeline/Source/MovieRenderPipelineEditor/Private/MoviePipelinePIEExecutor.cpp
+++ b/Engine/Plugins/MovieScene/MovieRenderPipeline/Source/MovieRenderPipelineEditor/Private/MoviePipelinePIEExecutor.cpp
@@ -21,6 +21,7 @@
 #include "Modules/ModuleManager.h"
 #include "MessageLogModule.h"
 #include "Logging/MessageLog.h"
+#include "Graph/MovieGraphBlueprintLibrary.h"
 #include "Graph/MovieGraphPipeline.h"
 #include "Graph/Nodes/MovieGraphGlobalGameOverrides.h"
 #include "LevelSequence.h"
@@ -130,12 +131,33 @@ void UMoviePipelinePIEExecutor::Start(const UMoviePipelineExecutorJob* InJob)
 
 	FVector2D WindowSize = FVector2D(1280, 720);
 
+	// Allow creating the PIE window at the size of the output to work around UMG scaling bugs when the main viewport is one resolution and then we re-render it at another.
 	const UMoviePipelinePIEExecutorSettings* ExecutorSettings = GetDefault<UMoviePipelinePIEExecutorSettings>();
 	if (ExecutorSettings->bResizePIEWindowToOutputResolution)
 	{
-		// Allow creating the PIE window at the size of the output to work around UMG scaling bugs when the main viewport is one resolution and then we re-render it at another.
-		UMoviePipelineOutputSetting* OutputSetting = Cast<UMoviePipelineOutputSetting>(InJob->GetConfiguration()->FindOrAddSettingByClass(UMoviePipelineOutputSetting::StaticClass()));
-		WindowSize = FVector2D(OutputSetting->OutputResolution.X, OutputSetting->OutputResolution.Y);
+		UMovieGraphConfig* GraphConfig = InJob->GetGraphPreset();
+		if (InJob->IsUsingGraphConfiguration() && GraphConfig)
+		{
+			// Like the GameMode below, this isn't a perfect evaluation, and is further complicated by things like Adapt Output Resolution and Overscan
+			// but it should work in the basic case.
+			FMovieGraphTraversalContext TraversalContext;
+			TraversalContext.Job = const_cast<UMoviePipelineExecutorJob*>(InJob);
+
+			FString OutTraversalError;
+			if (UMovieGraphEvaluatedConfig* EvaluatedGraph = GraphConfig->CreateFlattenedGraph(TraversalContext, OutTraversalError))
+			{
+				WindowSize = UMovieGraphBlueprintLibrary::GetOverscannedResolution(EvaluatedGraph);
+			}
+			else
+			{
+				UE_LOG(LogMovieRenderPipeline, Warning, TEXT("Failed to evaluate graph to determine output resolution. Traversal error: %s"), *OutTraversalError);
+			}
+		}
+		else
+		{
+			UMoviePipelineOutputSetting* OutputSetting = Cast<UMoviePipelineOutputSetting>(InJob->GetConfiguration()->FindOrAddSettingByClass(UMoviePipelineOutputSetting::StaticClass()));
+			WindowSize = FVector2D(OutputSetting->OutputResolution.X, OutputSetting->OutputResolution.Y);
+		}
 	}
 
 
