MoviePipeline: Fixed a crash in MRG that could occur if creating a reroute node via double-click, and a reroute node already existed immediately upstream or downstream.

diff --git a/Engine/Plugins/MovieScene/MovieRenderPipeline/Source/MovieRenderPipelineEditor/Private/Graph/MovieGraphSchema.cpp b/Engine/Plugins/MovieScene/MovieRenderPipeline/Source/MovieRenderPipelineEditor/Private/Graph/MovieGraphSchema.cpp
--- a/Engine/Plugins/MovieScene/MovieRenderPipeline/Source/MovieRenderPipelineEditor/Private/Graph/MovieGraphSchema.cpp
+++ b/Engine/Plugins/MovieScene/MovieRenderPipeline/Source/MovieRenderPipelineEditor/Private/Graph/MovieGraphSchema.cpp
@@ -256,13 +256,21 @@ void UMovieGraphSchema::OnPinConnectionDoubleCicked(UEdGraphPin* PinA, UEdGraphP
 		// Break the existing connection (ie, the connection that was clicked on).
 		PinA->BreakLinkTo(PinB);
 
-		// Create the reroute node. Using PinB as the "from" pin will create the downstream connection.
-		const UEdGraphNode* NewRerouteNode = RerouteNodeAction->PerformAction(OwningGraph, PinB, NodePosition);
+		// Create the reroute node. Use nullptr as the FromPin -- we'll manually perform connections via TryCreateConnection() so the
+		// pin type propagates properly and the runtime graph is informed of the change.
+		UEdGraphPin* FromPin = nullptr;
+		const UEdGraphNode* NewRerouteNode = RerouteNodeAction->PerformAction(OwningGraph, FromPin, NodePosition);
+
+		// Connect the reroute node's right (output) pin back to the previously-connected downstream node.
+		if (UEdGraphPin* OutputPin = NewRerouteNode->FindPin(NAME_None, EGPD_Output))
+		{
+			TryCreateConnection(OutputPin, PinB);
+		}
 
 		// Connect the reroute node's left (input) pin back to the previously-connected upstream node.
 		if (UEdGraphPin* InputPin = NewRerouteNode->FindPin(NAME_None, EGPD_Input))
 		{
-			InputPin->MakeLinkTo(PinA);
+			TryCreateConnection(PinA, InputPin);
 		}
 	}
 }
