MoviePipeline: Fixed a bug in MRG that caused the "Cast Shadow While Hidden" property on the Modifier node to not work in isolation (it required the "Affect Indirect Lighting While Hidden" property to take effect).

The problem was caused by an erroneous call to SetIsTemporarilyHiddenInEditor(), which has existed in MRG since the very first prototype of its layer system (which modified in-editor actors rather than in-game actors).

diff --git a/Engine/Plugins/MovieScene/MovieRenderPipeline/Source/MovieRenderPipelineCore/Private/Graph/MovieGraphRenderLayerSubsystem.cpp b/Engine/Plugins/MovieScene/MovieRenderPipeline/Source/MovieRenderPipelineCore/Private/Graph/MovieGraphRenderLayerSubsystem.cpp
--- a/Engine/Plugins/MovieScene/MovieRenderPipeline/Source/MovieRenderPipelineCore/Private/Graph/MovieGraphRenderLayerSubsystem.cpp
+++ b/Engine/Plugins/MovieScene/MovieRenderPipeline/Source/MovieRenderPipelineCore/Private/Graph/MovieGraphRenderLayerSubsystem.cpp
@@ -411,9 +411,6 @@ void UMovieGraphRenderPropertyModifier::SetActorAndComponentState(const TArray<F
 		}
 
 		InActor->SetActorHiddenInGame(InNewActorState.bIsHidden);
-#if WITH_EDITOR
-		InActor->SetIsTemporarilyHiddenInEditor(InNewActorState.bIsHidden);
-#endif
 	};
 
 	// If the node's state/settings should be applied, generate a singular FComponentState and apply it to all cached components.
