MoviePipeline: Fixed a regression in MRG that caused the "Overwrite Existing Output" property on the Global Output Settings node to not work correctly (no files were written if there were already files on disk).

diff --git a/Engine/Plugins/MovieScene/MovieRenderPipeline/Source/MovieRenderPipelineCore/Private/Graph/MovieGraphBlueprintLibrary.cpp b/Engine/Plugins/MovieScene/MovieRenderPipeline/Source/MovieRenderPipelineCore/Private/Graph/MovieGraphBlueprintLibrary.cpp
--- a/Engine/Plugins/MovieScene/MovieRenderPipeline/Source/MovieRenderPipelineCore/Private/Graph/MovieGraphBlueprintLibrary.cpp
+++ b/Engine/Plugins/MovieScene/MovieRenderPipeline/Source/MovieRenderPipelineCore/Private/Graph/MovieGraphBlueprintLibrary.cpp
@@ -223,8 +223,9 @@ FString UMovieGraphBlueprintLibrary::ResolveFilenameFormatArguments(const FStrin
 	
 	// If the extension is not resolved, ThisTry will be a path
 	const FString ExtToken = TEXT(".{ext}");
-	FString Extension = FString::Format(*ExtToken, NamedArgs);
-	FString ThisTry = Extension == ExtToken ? BaseFilename : BaseFilename + Extension;
+	const FString Extension = FString::Format(*ExtToken, NamedArgs);
+	const FString ResolvedExtension = (Extension == ExtToken) ? FString() : Extension;
+	FString ThisTry = BaseFilename + Extension;
 	
 	// Check that it's a valid path that we can write to.
 	if (UE::MoviePipeline::CanWriteToFile(*ThisTry, bOverwriteExisting))
@@ -235,14 +236,14 @@ FString UMovieGraphBlueprintLibrary::ResolveFilenameFormatArguments(const FStrin
 	// If we got here it means that there was already a file there with that name and bOverwriteExisting was false.
 	// So we start by swapping in _(2) where the file_dup token is, and try again. If that fails, _(3), etc. We start
 	// at _(2) because the file they're trying to write (foo.png) is technically "_(1)", so we start the rename process
-	// at _(2) so it accurately reflects that this is the second verison of that file.
+	// at _(2) so it accurately reflects that this is the second version of that file.
 	int32 DuplicateIndex = 2;
 	while (true)
 	{
 		NamedArgs.Add(TEXT("file_dup"), FString::Printf(TEXT("_(%d)"), DuplicateIndex));
 
 		// Re-resolve the format string now that we've reassigned frame_dup to a number.
-		ThisTry = FString::Format(*InFormatString, NamedArgs);
+		ThisTry = FString::Format(*InFormatString, NamedArgs) + ResolvedExtension;
 
 		// If the file doesn't exist, we can use that, else, increment the index and try again
 		if (UE::MoviePipeline::CanWriteToFile(*ThisTry, bOverwriteExisting))
