MoviePipeline: Fixed an issue in Quick Render that would cause the wrong camera to be rendered if the Current Viewport mode is active, and the viewport is locked to a shot camera from Sequencer.

diff --git a/Engine/Plugins/MovieScene/MovieRenderPipeline/Source/MovieRenderPipelineEditor/Private/Graph/MovieGraphQuickRender.cpp b/Engine/Plugins/MovieScene/MovieRenderPipeline/Source/MovieRenderPipelineEditor/Private/Graph/MovieGraphQuickRender.cpp
--- a/Engine/Plugins/MovieScene/MovieRenderPipeline/Source/MovieRenderPipelineEditor/Private/Graph/MovieGraphQuickRender.cpp
+++ b/Engine/Plugins/MovieScene/MovieRenderPipeline/Source/MovieRenderPipelineEditor/Private/Graph/MovieGraphQuickRender.cpp
@@ -743,6 +743,19 @@ void UMovieGraphQuickRenderSubsystem::CachePrePieData(UWorld* InEditorWorld)
 	TArray<ACameraActor*> CameraActors;
 	GEditor->GetSelectedActors()->GetSelectedObjects<ACameraActor>(CameraActors);
 	CopyToWeakArray(CachedPrePieData.SelectedCameras, CameraActors);
+
+	// If the viewport is locked to an actor (likely a camera), cache it here. It may temporarily change while PIE is being started, so it needs to
+	// be determined before PIE starts up.
+	if (const FLevelEditorViewportClient* ViewportClient = UMovieGraphApplyViewportLookNode::GetViewportClient())
+	{
+		const FLevelViewportActorLock& ActorLock = ViewportClient->IsLockedToCinematic() ? ViewportClient->GetCinematicActorLock() : ViewportClient->GetActorLock();
+		if (ActorLock.HasValidLockedActor())
+		{
+			CachedPrePieData.ViewportActorLock = ActorLock.LockedActor;
+		}
+		
+		CachedPrePieData.ViewportActorLockCameraComponent = ViewportClient->GetCameraComponentForView();
+	}
 }
 
 void UMovieGraphQuickRenderSubsystem::RestorePreRenderState()
@@ -796,6 +809,8 @@ void UMovieGraphQuickRenderSubsystem::RestorePreRenderState()
 
 	CachedPrePieData.TemporarilyNonTransientComponents.Empty();
 	CachedPrePieData.TemporarilyNonTransientActors.Empty();
+	CachedPrePieData.ViewportActorLock.Reset();
+	CachedPrePieData.ViewportActorLockCameraComponent = nullptr;
 }
 
 ULevelSequence* UMovieGraphQuickRenderSubsystem::SetUpRenderingLevelSequence() const
@@ -1081,18 +1096,19 @@ void UMovieGraphQuickRenderSubsystem::PerformPreRenderSetup(UWorld* InEditorWorl
 			
 			// If the viewport is being piloted by a camera, then we should use the properties of that to render from instead of the viewport's
 			// non-piloted camera properties.
-			const FLevelViewportActorLock& ActorLock = ViewportClient->GetActorLock();
-			if (ActorLock.HasValidLockedActor())
+			if (CachedPrePieData.ViewportActorLock.IsValid())
 			{
 				// Support basic position/rotation for non-camera actors
-				const TStrongObjectPtr<AActor> LockedActorPin = ActorLock.LockedActor.Pin();
+				const TStrongObjectPtr<AActor> LockedActorPin = CachedPrePieData.ViewportActorLock.Pin();
 				CameraPosition = LockedActorPin->GetActorLocation();
 				CameraRotation = LockedActorPin->GetActorRotation();
 
 				// If the piloted actor is a camera, get all of its camera properties. Note that this supports "cinematic locks" (like cameras piloted
 				// via Sequencer) and normal piloted cameras.
-				if (const UCameraComponent* PilotCameraComponent = ViewportClient->GetCameraComponentForView())
+				if (IsValid(CachedPrePieData.ViewportActorLockCameraComponent))
 				{
+					const UCameraComponent* PilotCameraComponent = CachedPrePieData.ViewportActorLockCameraComponent;
+					
 					CameraPosition = PilotCameraComponent->GetTypedOuter<AActor>()->GetActorLocation();
 					CameraRotation = PilotCameraComponent->GetTypedOuter<AActor>()->GetActorRotation();
 					bIsPerspective = PilotCameraComponent->ProjectionMode == ECameraProjectionMode::Perspective;
diff --git a/Engine/Plugins/MovieScene/MovieRenderPipeline/Source/MovieRenderPipelineEditor/Public/Graph/MovieGraphQuickRender.h b/Engine/Plugins/MovieScene/MovieRenderPipeline/Source/MovieRenderPipelineEditor/Public/Graph/MovieGraphQuickRender.h
--- a/Engine/Plugins/MovieScene/MovieRenderPipeline/Source/MovieRenderPipelineEditor/Public/Graph/MovieGraphQuickRender.h
+++ b/Engine/Plugins/MovieScene/MovieRenderPipeline/Source/MovieRenderPipelineEditor/Public/Graph/MovieGraphQuickRender.h
@@ -65,6 +65,12 @@ private:
 
 		/** Components that have been temporarily flagged as RF_NonPIEDuplicateTransient. This enables them to be copied to the PIE world when they typically would not. */
 		TArray<TWeakObjectPtr<UActorComponent>> TemporarilyNonTransientComponents;
+
+		/** The actor that the viewport is locked to (also referred to as the pilot camera). Nullptr if no actor lock is active. */
+		TWeakObjectPtr<AActor> ViewportActorLock = nullptr;
+
+		/** If the viewport is locked to an actor (see ViewportLockActor), this is the specific camera component that it is locked to. May be nullptr if the viewport is locked to a non-camera actor (like a light). */
+		UCameraComponent* ViewportActorLockCameraComponent = nullptr;
 	};
 	
 	/**
