MoviePipeline: Fixed issue that caused MRQ to not shut down correctly if a zero-range custom frame range was specified.

Also updated the executors to ensure that the correct pipeline is used if the pipeline specified in Project Settings is not compatible with the configuration type (ie, preset or graph) specified on the job.

diff --git a/Engine/Plugins/MovieScene/MovieRenderPipeline/Source/MovieRenderPipelineCore/Private/MoviePipeline.cpp b/Engine/Plugins/MovieScene/MovieRenderPipeline/Source/MovieRenderPipelineCore/Private/MoviePipeline.cpp
--- a/Engine/Plugins/MovieScene/MovieRenderPipeline/Source/MovieRenderPipelineCore/Private/MoviePipeline.cpp
+++ b/Engine/Plugins/MovieScene/MovieRenderPipeline/Source/MovieRenderPipelineCore/Private/MoviePipeline.cpp
@@ -200,14 +200,6 @@ void UMoviePipeline::Initialize(UMoviePipelineExecutorJob* InJob)
 #endif
 			TargetSequence->GetMovieScene()->SetPlaybackRange(CustomPlaybackRange);
 		}
-
-		// Warn about zero length playback ranges, often happens because they set the Start/End frame to the same frame.
-		if (TargetSequence->GetMovieScene()->GetPlaybackRange().IsEmpty())
-		{
-			UE_LOG(LogMovieRenderPipeline, Error, TEXT("Playback Range was zero. End Frames are exclusive, did you mean [n, n+1]?"));
-			Shutdown(true);
-			return;
-		}
 	}
 	
 	// Initialize all of our primary config settings. Shot specific ones will be called for their appropriate shot.
@@ -319,12 +311,24 @@ void UMoviePipeline::Initialize(UMoviePipelineExecutorJob* InJob)
 		Shot->ShotInfo.VersionNumber = UMoviePipelineBlueprintLibrary::ResolveVersionNumber(Params);
 	}
 
-	// If the shot mask entirely disabled everything we'll transition directly to finish as there is no work to do.
-	if (ActiveShotList.Num() == 0)
+	// Warn about zero length playback ranges, often happens because they set the Start/End frame to the same frame.
+	const bool bIsEmptyPlaybackRange = TargetSequence->GetMovieScene()->GetPlaybackRange().IsEmpty();
+	if (bIsEmptyPlaybackRange)
+	{
+		UE_LOG(LogMovieRenderPipeline, Error, TEXT("Playback Range was zero. End Frames are exclusive, did you mean [n, n+1]?"));
+	}
+
+	// Warn if the shot mask entirely disabled everything.
+	const bool bAllShotsDisabled = ActiveShotList.IsEmpty();
+	if (bAllShotsDisabled)
 	{
-		// We have to transition twice as Uninitialized -> n state is a no-op, so the second tick will take us to Finished which shuts down.
 		UE_LOG(LogMovieRenderPipeline, Warning, TEXT("[%d] No shots detected to render. Either all outside playback range, or disabled via shot mask, bailing."), GFrameCounter);
+	}
 
+	// If there's no work to do, shut down.
+	if (bIsEmptyPlaybackRange || bAllShotsDisabled)
+	{
+		// We have to transition twice as Uninitialized -> n state is a no-op, so the second tick will take us to Finished which shuts down.
 		TransitionToState(EMovieRenderPipelineState::Export);
 		TransitionToState(EMovieRenderPipelineState::Finished);
 	}
diff --git a/Engine/Plugins/MovieScene/MovieRenderPipeline/Source/MovieRenderPipelineCore/Private/MoviePipelineInProcessExecutor.cpp b/Engine/Plugins/MovieScene/MovieRenderPipeline/Source/MovieRenderPipelineCore/Private/MoviePipelineInProcessExecutor.cpp
--- a/Engine/Plugins/MovieScene/MovieRenderPipeline/Source/MovieRenderPipelineCore/Private/MoviePipelineInProcessExecutor.cpp
+++ b/Engine/Plugins/MovieScene/MovieRenderPipeline/Source/MovieRenderPipelineCore/Private/MoviePipelineInProcessExecutor.cpp
@@ -117,8 +117,8 @@ void UMoviePipelineInProcessExecutor::OnMapLoadFinished(UWorld* NewWorld)
 		MoviePipelineClass = UMoviePipeline::StaticClass();
 	}
 
-	// Temporary
-	if (CurrentJob->IsUsingGraphConfiguration())
+	// Force the pipeline to the default graph pipeline if we're using a graph and the pipeline isn't graph-compatible
+	if (CurrentJob->IsUsingGraphConfiguration() && !MoviePipelineClass->IsChildOf(UMovieGraphPipeline::StaticClass()))
 	{
 		MoviePipelineClass = UMovieGraphPipeline::StaticClass();
 	}
diff --git a/Engine/Plugins/MovieScene/MovieRenderPipeline/Source/MovieRenderPipelineEditor/Private/MoviePipelinePIEExecutor.cpp b/Engine/Plugins/MovieScene/MovieRenderPipeline/Source/MovieRenderPipelineEditor/Private/MoviePipelinePIEExecutor.cpp
--- a/Engine/Plugins/MovieScene/MovieRenderPipeline/Source/MovieRenderPipelineEditor/Private/MoviePipelinePIEExecutor.cpp
+++ b/Engine/Plugins/MovieScene/MovieRenderPipeline/Source/MovieRenderPipelineEditor/Private/MoviePipelinePIEExecutor.cpp
@@ -255,11 +255,19 @@ void UMoviePipelinePIEExecutor::OnPIEStartupFinished(bool)
 		return;
 	}
 
+	// Ensure the pipeline class we're using is compatible with the job's configuration type
 	const UMoviePipelineExecutorJob* CurrentJob = Queue->GetJobs()[CurrentPipelineIndex];
-	if (CurrentJob->IsUsingGraphConfiguration())
+	const bool bIsUsingGraph = CurrentJob->IsUsingGraphConfiguration();
+	if (bIsUsingGraph && !PipelineClass->IsChildOf(UMovieGraphPipeline::StaticClass()))
 	{
+		UE_LOG(LogMovieRenderPipeline, Warning, TEXT("Using the default MovieGraphPipeline; the pipeline specified in Project Settings (%s) is not compatible with graphs."), *PipelineClass->GetName());
 		PipelineClass = UMovieGraphPipeline::StaticClass();
 	}
+	else if (!bIsUsingGraph && !PipelineClass->IsChildOf(UMoviePipeline::StaticClass()))
+	{
+		UE_LOG(LogMovieRenderPipeline, Warning, TEXT("Using the default MoviePipeline; the pipeline specified in Project Settings (%s) is not compatible with normal (non-graph) presets."), *PipelineClass->GetName());
+		PipelineClass = UMoviePipeline::StaticClass();
+	}
 
 	// This Pipeline belongs to the world being created so that they have context for things they execute.
 	ActiveMoviePipeline = NewObject<UMoviePipelineBase>(ExecutingWorld, PipelineClass, NAME_None, RF_Transient);
