MoviePipeline: Updated Quick Render to correctly show control rig actors when the "Show Editor-Only Actors" viewport visibility option is turned on (without needing the "Visibility" flag also turned on).

diff --git a/Engine/Plugins/MovieScene/MovieRenderPipeline/Source/MovieRenderPipelineEditor/Private/Graph/MovieGraphQuickRender.cpp b/Engine/Plugins/MovieScene/MovieRenderPipeline/Source/MovieRenderPipelineEditor/Private/Graph/MovieGraphQuickRender.cpp
--- a/Engine/Plugins/MovieScene/MovieRenderPipeline/Source/MovieRenderPipelineEditor/Private/Graph/MovieGraphQuickRender.cpp
+++ b/Engine/Plugins/MovieScene/MovieRenderPipeline/Source/MovieRenderPipelineEditor/Private/Graph/MovieGraphQuickRender.cpp
@@ -474,7 +474,7 @@ void UMovieGraphQuickRenderSubsystem::ApplyQuickRenderUpdatesToDuplicateGraph_Ap
 			}
 		}
 
-		if (Actor->IsEditorOnly() || bHasEditorOnlyComponent)
+		if (Actor->IsEditorOnly() || Actor->bIsEditorPreviewActor || bHasEditorOnlyComponent)
 		{
 			FMovieGraphActorQueryEntry& VisibleActorEntry = VisibleActors.AddDefaulted_GetRef();
 			VisibleActorEntry.ActorToMatch = Actor;
