Add boolean for MoviePipelineFunctionalTestBase to toggle screenshot comparisons

diff --git a/Engine/Plugins/MovieScene/MovieRenderPipeline/Source/MovieRenderPipelineEditor/Private/MoviePipelineFunctionalTestBase.cpp b/Engine/Plugins/MovieScene/MovieRenderPipeline/Source/MovieRenderPipelineEditor/Private/MoviePipelineFunctionalTestBase.cpp
--- a/Engine/Plugins/MovieScene/MovieRenderPipeline/Source/MovieRenderPipelineEditor/Private/MoviePipelineFunctionalTestBase.cpp
+++ b/Engine/Plugins/MovieScene/MovieRenderPipeline/Source/MovieRenderPipelineEditor/Private/MoviePipelineFunctionalTestBase.cpp
@@ -25,6 +25,7 @@
 AMoviePipelineFunctionalTestBase::AMoviePipelineFunctionalTestBase()
 {
 	ImageToleranceLevel = EImageTolerancePreset::IgnoreLess;
+	bPerformDiff = true;
 }
 
 
@@ -91,7 +92,10 @@ void AMoviePipelineFunctionalTestBase::OnMoviePipelineFinished(FMoviePipelineOut
 	}
 	else
 	{
-		CompareRenderOutputToGroundTruth(InOutputData);
+		if (bPerformDiff)
+		{
+			CompareRenderOutputToGroundTruth(InOutputData);
+		}
 	}
 }
 
diff --git a/Engine/Plugins/MovieScene/MovieRenderPipeline/Source/MovieRenderPipelineEditor/Private/MoviePipelineFunctionalTestBase.h b/Engine/Plugins/MovieScene/MovieRenderPipeline/Source/MovieRenderPipelineEditor/Private/MoviePipelineFunctionalTestBase.h
--- a/Engine/Plugins/MovieScene/MovieRenderPipeline/Source/MovieRenderPipelineEditor/Private/MoviePipelineFunctionalTestBase.h
+++ b/Engine/Plugins/MovieScene/MovieRenderPipeline/Source/MovieRenderPipelineEditor/Private/MoviePipelineFunctionalTestBase.h
@@ -51,6 +51,9 @@ protected:
 
 	UPROPERTY(EditAnywhere, Category = "Movie Pipeline")
 	FImageTolerance CustomImageTolerance;
+	
+	UPROPERTY(EditAnywhere, Category = "Movie Pipeline")
+	bool bPerformDiff;
 
 protected:
 	UPROPERTY(Transient)
