MoviePipeline: Fixed a crash in the MP4 node when the resolution exceeded 4k.

Also turned a warning about mismatched pipline types into an info log instead; generally, the user doesn't need to know about this.

diff --git a/Engine/Plugins/MovieScene/MovieRenderPipeline/Source/MovieRenderPipelineCore/Private/Graph/MovieGraphPipeline.cpp b/Engine/Plugins/MovieScene/MovieRenderPipeline/Source/MovieRenderPipelineCore/Private/Graph/MovieGraphPipeline.cpp
--- a/Engine/Plugins/MovieScene/MovieRenderPipeline/Source/MovieRenderPipelineCore/Private/Graph/MovieGraphPipeline.cpp
+++ b/Engine/Plugins/MovieScene/MovieRenderPipeline/Source/MovieRenderPipelineCore/Private/Graph/MovieGraphPipeline.cpp
@@ -588,12 +588,13 @@ void UMovieGraphPipeline::TickProducingFrames()
 	//
 	//ProcessOutstandingFutures();
 	//
-	//if (bShutdownRequested)
-	//{
-	//	UE_LOG(LogMovieRenderPipeline, Log, TEXT("MovieGraph TickProductingFrames: Async Shutdown Requested, abandoning remaining work and moving to Finalize."));
-	//	TransitionToState(EMovieRenderPipelineState::Finalize);
-	//	return;
-	//}
+	
+	if (bShutdownRequested)
+	{
+		UE_LOG(LogMovieRenderPipeline, Log, TEXT("MovieGraph TickProducingFrames: Async Shutdown Requested, abandoning remaining work and moving to Finalize."));
+		TransitionToState(EMovieRenderPipelineState::Finalize);
+		return;
+	}
 
 	// When start up we want to override the engine's Custom Timestep with our own.
 	// This gives us the ability to completely control the engine tick/delta time before the frame
diff --git a/Engine/Plugins/MovieScene/MovieRenderPipeline/Source/MovieRenderPipelineCore/Private/Graph/Nodes/MovieGraphVideoOutputNode.cpp b/Engine/Plugins/MovieScene/MovieRenderPipeline/Source/MovieRenderPipelineCore/Private/Graph/Nodes/MovieGraphVideoOutputNode.cpp
--- a/Engine/Plugins/MovieScene/MovieRenderPipeline/Source/MovieRenderPipelineCore/Private/Graph/Nodes/MovieGraphVideoOutputNode.cpp
+++ b/Engine/Plugins/MovieScene/MovieRenderPipeline/Source/MovieRenderPipelineCore/Private/Graph/Nodes/MovieGraphVideoOutputNode.cpp
@@ -82,8 +82,9 @@ void UMovieGraphVideoOutputNode::OnReceiveImageDataImpl(UMovieGraphPipeline* InP
 			continue;
 		}
 
+		// GetOrCreateOutputWriter() can set bHasError, so even if a writer was created, don't continue if there was an error initializing it
 		const FMovieGraphCodecWriterWithPromise* OutputWriter = GetOrCreateOutputWriter(InPipeline, InRawFrameData, RenderPassData, CompositedPasses);
-		if (!OutputWriter)
+		if (!OutputWriter || bHasError)
 		{
 			continue;
 		}
@@ -210,7 +211,12 @@ void UMovieGraphVideoOutputNode::OnAllFramesFinalizedImpl(UMovieGraphPipeline* I
 
 		Finalize_EncodeThread(RawWriter);
 
-		Writer.Promise.SetValue(!bHasError);
+		// Note: If there was an error, the promise should already be set to false, so we only set it to true here if there's no error. This isn't
+		// ideal and it would be nice to have another way of tracking this, but double-setting the promise will trigger a check().
+		if (!bHasError)
+		{
+			Writer.Promise.SetValue(!bHasError);
+		}
 	}
 
 	// AllWriters is static and shared across *all* video output node types. Only delete writers of the current node's type.
diff --git a/Engine/Plugins/MovieScene/MovieRenderPipeline/Source/MovieRenderPipelineEditor/Private/MoviePipelinePIEExecutor.cpp b/Engine/Plugins/MovieScene/MovieRenderPipeline/Source/MovieRenderPipelineEditor/Private/MoviePipelinePIEExecutor.cpp
--- a/Engine/Plugins/MovieScene/MovieRenderPipeline/Source/MovieRenderPipelineEditor/Private/MoviePipelinePIEExecutor.cpp
+++ b/Engine/Plugins/MovieScene/MovieRenderPipeline/Source/MovieRenderPipelineEditor/Private/MoviePipelinePIEExecutor.cpp
@@ -260,12 +260,12 @@ void UMoviePipelinePIEExecutor::OnPIEStartupFinished(bool)
 	const bool bIsUsingGraph = CurrentJob->IsUsingGraphConfiguration();
 	if (bIsUsingGraph && !PipelineClass->IsChildOf(UMovieGraphPipeline::StaticClass()))
 	{
-		UE_LOG(LogMovieRenderPipeline, Warning, TEXT("Using the default MovieGraphPipeline; the pipeline specified in Project Settings (%s) is not compatible with graphs."), *PipelineClass->GetName());
+		UE_LOG(LogMovieRenderPipeline, Log, TEXT("Using the default MovieGraphPipeline; the pipeline specified in Project Settings (%s) is not compatible with graphs."), *PipelineClass->GetName());
 		PipelineClass = UMovieGraphPipeline::StaticClass();
 	}
 	else if (!bIsUsingGraph && !PipelineClass->IsChildOf(UMoviePipeline::StaticClass()))
 	{
-		UE_LOG(LogMovieRenderPipeline, Warning, TEXT("Using the default MoviePipeline; the pipeline specified in Project Settings (%s) is not compatible with normal (non-graph) presets."), *PipelineClass->GetName());
+		UE_LOG(LogMovieRenderPipeline, Log, TEXT("Using the default MoviePipeline; the pipeline specified in Project Settings (%s) is not compatible with normal (non-graph) presets."), *PipelineClass->GetName());
 		PipelineClass = UMoviePipeline::StaticClass();
 	}
 
diff --git a/Engine/Plugins/MovieScene/MovieRenderPipeline/Source/MovieRenderPipelineMP4Encoder/Private/Graph/MovieGraphMP4EncoderNode.cpp b/Engine/Plugins/MovieScene/MovieRenderPipeline/Source/MovieRenderPipelineMP4Encoder/Private/Graph/MovieGraphMP4EncoderNode.cpp
--- a/Engine/Plugins/MovieScene/MovieRenderPipeline/Source/MovieRenderPipelineMP4Encoder/Private/Graph/MovieGraphMP4EncoderNode.cpp
+++ b/Engine/Plugins/MovieScene/MovieRenderPipeline/Source/MovieRenderPipelineMP4Encoder/Private/Graph/MovieGraphMP4EncoderNode.cpp
@@ -96,7 +96,7 @@ bool UMovieGraphMP4EncoderNode::Initialize_EncodeThread(MovieRenderGraph::IVideo
 	const FMP4CodecWriter* CodecWriter = static_cast<FMP4CodecWriter*>(InWriter);
 	if(!CodecWriter->Writer->Initialize())
 	{
-		UE_LOG(LogMovieRenderPipeline, Error, TEXT("Failed to initialize Movie Pipeline MP4 writer."));
+		UE_LOG(LogMovieRenderPipeline, Error, TEXT("Failed to initialize Movie Pipeline MP4 writer. This can sometimes be caused by trying to encode to a resolution higher than 4K (4096), which is unsupported."));
 		return false;
 	}
 	
diff --git a/Engine/Plugins/MovieScene/MovieRenderPipeline/Source/MovieRenderPipelineMP4Encoder/Private/Windows/MoviePipelineMP4Encoder.cpp b/Engine/Plugins/MovieScene/MovieRenderPipeline/Source/MovieRenderPipelineMP4Encoder/Private/Windows/MoviePipelineMP4Encoder.cpp
--- a/Engine/Plugins/MovieScene/MovieRenderPipeline/Source/MovieRenderPipelineMP4Encoder/Private/Windows/MoviePipelineMP4Encoder.cpp
+++ b/Engine/Plugins/MovieScene/MovieRenderPipeline/Source/MovieRenderPipelineMP4Encoder/Private/Windows/MoviePipelineMP4Encoder.cpp
@@ -689,7 +689,7 @@ static HRESULT GetVideoStreamEncoderAttributes(IMFAttributes** OutEncoderAttribu
 			}
 			if (SUCCEEDED(Result))
 			{
-				Result = (*OutEncoderAttributes)->SetUINT64(CODECAPI_AVEncCommonQuality, 0);
+				Result = (*OutEncoderAttributes)->SetUINT32(CODECAPI_AVEncCommonQuality, 0);
 			}
 		}
 		else if (InOptions.EncodingRateControl == EMoviePipelineMP4EncodeRateControlMode::ConstantQP)
