MoviePipeline: Fixed issue where burn-ins don't appear on MP4s in MRQ (when using the new H.264 MP4 setting).

diff --git a/Engine/Plugins/MovieScene/MovieRenderPipeline/Source/MovieRenderPipelineMP4Encoder/Private/MoviePipelineMP4EncoderOutput.cpp b/Engine/Plugins/MovieScene/MovieRenderPipeline/Source/MovieRenderPipelineMP4Encoder/Private/MoviePipelineMP4EncoderOutput.cpp
--- a/Engine/Plugins/MovieScene/MovieRenderPipeline/Source/MovieRenderPipelineMP4Encoder/Private/MoviePipelineMP4EncoderOutput.cpp
+++ b/Engine/Plugins/MovieScene/MovieRenderPipeline/Source/MovieRenderPipelineMP4Encoder/Private/MoviePipelineMP4EncoderOutput.cpp
@@ -62,7 +62,35 @@ void UMoviePipelineMP4EncoderOutput::WriteFrame_EncodeThread(MovieRenderPipeline
 	
 	// Quantize our 16 bit float data to 8 bit.
 	TUniquePtr<FImagePixelData> QuantizedPixelData = UE::MoviePipeline::QuantizeImagePixelDataToBitDepth(InPixelData, 8);
-	 
+
+	// Do a quick composite of renders/burn-ins.
+	TArray<FPixelPreProcessor> PixelPreProcessors;
+	for (const MoviePipeline::FCompositePassInfo& CompositePass : InCompositePasses)
+	{
+		// We don't need to copy the data here (even though it's being passed to a async system) because we already made a unique copy of the
+		// burn in/widget data when we decided to composite it.
+		switch (QuantizedPixelData->GetType())
+		{
+		case EImagePixelType::Color:
+			PixelPreProcessors.Add(TAsyncCompositeImage<FColor>(CompositePass.PixelData->MoveImageDataToNew()));
+			break;
+		case EImagePixelType::Float16:
+			PixelPreProcessors.Add(TAsyncCompositeImage<FFloat16Color>(CompositePass.PixelData->MoveImageDataToNew()));
+			break;
+		case EImagePixelType::Float32:
+			PixelPreProcessors.Add(TAsyncCompositeImage<FLinearColor>(CompositePass.PixelData->MoveImageDataToNew()));
+			break;
+		}
+	}
+
+	// This is done on the main thread for simplicity but the composite itself is parallelized.
+	FImagePixelData* PixelData = QuantizedPixelData.Get();
+	for (const FPixelPreProcessor& PreProcessor : PixelPreProcessors)
+	{
+		// PreProcessors are assumed to be valid.
+		PreProcessor(PixelData);
+	}
+
 	const void* Data = nullptr;
 	int64 DataSize;
 	QuantizedPixelData->GetRawData(Data, DataSize);
