MoviePipeline: Pared down the number of options available on the MP4 node/setting in MRG and MRQ. These removed options will still be available for scripting to control, but have been removed from the UI because it's expected the vast majority of users will not need to control them.

diff --git a/Engine/Plugins/MovieScene/MovieRenderPipeline/Source/MovieRenderPipelineMP4Encoder/Private/Graph/MovieGraphMP4EncoderNode.cpp b/Engine/Plugins/MovieScene/MovieRenderPipeline/Source/MovieRenderPipelineMP4Encoder/Private/Graph/MovieGraphMP4EncoderNode.cpp
--- a/Engine/Plugins/MovieScene/MovieRenderPipeline/Source/MovieRenderPipelineMP4Encoder/Private/Graph/MovieGraphMP4EncoderNode.cpp
+++ b/Engine/Plugins/MovieScene/MovieRenderPipeline/Source/MovieRenderPipelineMP4Encoder/Private/Graph/MovieGraphMP4EncoderNode.cpp
@@ -72,15 +72,26 @@ TUniquePtr<MovieRenderGraph::IVideoCodecWriter> UMovieGraphMP4EncoderNode::Initi
 	Options.Width = InInitializationContext.Resolution.X;
 	Options.Height = InInitializationContext.Resolution.Y;
 	Options.FrameRate = EffectiveFrameRate;
-	Options.CommonMaxBitRate = FMath::RoundToInt(EvaluatedNode->MaxBitrateInMbps*1024*1024);
 	Options.CommonMeanBitRate = FMath::RoundToInt(EvaluatedNode->AverageBitrateInMbps*1024*1024);
 	Options.CommonQualityVsSpeed = 100;
 	Options.CommonConstantRateFactor = EvaluatedNode->ConstantRateFactor;
-	Options.EncodingProfile = EvaluatedNode->EncodingProfile;
-	Options.EncodingLevel = EvaluatedNode->EncodingLevel;
 	Options.EncodingRateControl = EvaluatedNode->EncodingRateControl;
 	Options.bIncludeAudio = EvaluatedNode->bIncludeAudio;
 	
+	// Optional properties that can be overridden by scripting if needed; not exposed to the UI currently
+	if (EvaluatedNode->bOverride_MaxBitrateInMbps)
+	{
+		Options.CommonMaxBitRate = FMath::RoundToInt(EvaluatedNode->MaxBitrateInMbps*1024*1024);
+	}
+	if (EvaluatedNode->bOverride_EncodingProfile)
+	{
+		Options.EncodingProfile = EvaluatedNode->EncodingProfile;
+	}
+	if (EvaluatedNode->bOverride_EncodingLevel)
+	{
+		Options.EncodingLevel = EvaluatedNode->EncodingLevel;
+	}
+	
 	TUniquePtr<FMP4CodecWriter> NewWriter = MakeUnique<FMP4CodecWriter>();
 	NewWriter->Writer = MakeUnique<FMoviePipelineMP4Encoder>(Options);
 
diff --git a/Engine/Plugins/MovieScene/MovieRenderPipeline/Source/MovieRenderPipelineMP4Encoder/Private/Graph/MovieGraphMP4EncoderNodeCustomization.h b/Engine/Plugins/MovieScene/MovieRenderPipeline/Source/MovieRenderPipelineMP4Encoder/Private/Graph/MovieGraphMP4EncoderNodeCustomization.h
new file mode 100644
--- /dev/null
+++ b/Engine/Plugins/MovieScene/MovieRenderPipeline/Source/MovieRenderPipelineMP4Encoder/Private/Graph/MovieGraphMP4EncoderNodeCustomization.h
@@ -0,0 +1,69 @@
+// Copyright Epic Games, Inc. All Rights Reserved.
+
+#pragma once
+
+#if WITH_EDITOR
+
+#include "Graph/MovieGraphMP4EncoderNode.h"
+
+#include "DetailLayoutBuilder.h"
+#include "IDetailCustomization.h"
+#include "PropertyHandle.h"
+
+#define LOCTEXT_NAMESPACE "MovieRenderPipelineMP4Encoder"
+
+/** Customize how the MP4 Encoder node appears in the details panel. */
+class FMovieGraphMP4EncoderNodeNodeCustomization : public IDetailCustomization
+{
+public:
+	static TSharedRef<IDetailCustomization> MakeInstance()
+	{
+		return MakeShared<FMovieGraphMP4EncoderNodeNodeCustomization>();
+	}
+
+protected:
+	//~ Begin IDetailCustomization interface
+	virtual void CustomizeDetails(const TSharedPtr<IDetailLayoutBuilder>& DetailBuilder) override
+	{
+		CustomizeDetails(*DetailBuilder);
+	}
+	
+	virtual void CustomizeDetails(IDetailLayoutBuilder& DetailBuilder) override
+	{
+		for (const TWeakObjectPtr<UMovieGraphMP4EncoderNode>& WeakMP4Node : DetailBuilder.GetObjectsOfTypeBeingCustomized<UMovieGraphMP4EncoderNode>())
+		{
+			const TAttribute<EVisibility> AverageBitrateVisibleAttr = TAttribute<EVisibility>::Create([WeakMP4Node]()
+			{
+				if (const TStrongObjectPtr<UMovieGraphMP4EncoderNode> MP4NodePin = WeakMP4Node.Pin())
+				{
+					return (MP4NodePin->EncodingRateControl == EMoviePipelineMP4EncodeRateControlMode::VariableBitRate) ? EVisibility::Visible : EVisibility::Collapsed;
+				}
+				
+				return EVisibility::Collapsed;
+			});
+
+			const TAttribute<EVisibility> ConstantRateFactorVisibleAttr = TAttribute<EVisibility>::Create([WeakMP4Node]()
+			{
+				if (const TStrongObjectPtr<UMovieGraphMP4EncoderNode> MP4NodePin = WeakMP4Node.Pin())
+				{
+					return (MP4NodePin->EncodingRateControl == EMoviePipelineMP4EncodeRateControlMode::Quality) ? EVisibility::Visible : EVisibility::Collapsed;
+				}
+				
+				return EVisibility::Collapsed;
+			});
+
+			const TSharedRef<IPropertyHandle> AverageBitrateHandle = DetailBuilder.GetProperty(GET_MEMBER_NAME_CHECKED(UMovieGraphMP4EncoderNode, AverageBitrateInMbps));
+			DetailBuilder.EditDefaultProperty(AverageBitrateHandle)
+				->Visibility(AverageBitrateVisibleAttr);
+
+			const TSharedRef<IPropertyHandle> ConstantRateFactorHandle = DetailBuilder.GetProperty(GET_MEMBER_NAME_CHECKED(UMovieGraphMP4EncoderNode, ConstantRateFactor));
+			DetailBuilder.EditDefaultProperty(ConstantRateFactorHandle)
+				->Visibility(ConstantRateFactorVisibleAttr);
+		}
+	}
+	//~ End IDetailCustomization interface
+};
+
+#undef LOCTEXT_NAMESPACE
+
+#endif	// WITH_EDITOR
\ No newline at end of file
diff --git a/Engine/Plugins/MovieScene/MovieRenderPipeline/Source/MovieRenderPipelineMP4Encoder/Private/MoviePipelineMP4EncoderCommon.h b/Engine/Plugins/MovieScene/MovieRenderPipeline/Source/MovieRenderPipelineMP4Encoder/Private/MoviePipelineMP4EncoderCommon.h
--- a/Engine/Plugins/MovieScene/MovieRenderPipeline/Source/MovieRenderPipelineMP4Encoder/Private/MoviePipelineMP4EncoderCommon.h
+++ b/Engine/Plugins/MovieScene/MovieRenderPipeline/Source/MovieRenderPipelineMP4Encoder/Private/MoviePipelineMP4EncoderCommon.h
@@ -53,15 +53,19 @@ UENUM(BlueprintType)
 enum class EMoviePipelineMP4EncodeRateControlMode : uint8
 {
 	/** Automatically chooses a bit rate to target the given quality. Valid ranges are 16-51. */
-	ConstantQP,
-	/** Automatically chooses a bit rate to target the given quality. Valid ranges are 16-51. */
+	ConstantQP UMETA(Hidden),
+
+	/** Automatically chooses a bit rate to target the given quality. */
 	Quality,
+
+	/** Attempts to achieve a given mean bitrate for every frame. Can result in higher bitrates than necessary on simple frames, and lower bitrates than required on complex frames. Unconstrained and single-pass. */
+	VariableBitRate,
+
 	/** Uses a variable bitrate that attempts to achieve a given mean bitrate, but can use a higher bitrate (with no max) on a given frame if needed. Value is in bytes per second.*/
-	VariableBitRate_Constrained,
-	/** Attempts to achieve a given mean bitrate for every frame. Can result in higher bitrates than necessary on simple frames, and lower bitrates than required on complex frames. Value is in bytes per second.*/
-	VariableBitRate_Unconstrained,
+	VariableBitRate_Constrained UMETA(Hidden),
+
 	/** Uses a variable bitrate that attempts to achieve a given mean bitrate, but can specify a maximum bitrate at which point quality will drop if needed. Value is in bytes per second.*/
-	ConstantBitRate,
+	ConstantBitRate UMETA(Hidden),
 }; 
 
 struct FMoviePipelineMP4EncoderOptions
@@ -76,12 +80,12 @@ struct FMoviePipelineMP4EncoderOptions
 		, AudioSampleRate(48000)
 		, AudioAverageBitRate(24000) // 192 kbps (bits)
 		, CommonMeanBitRate(12 * 1024 * 1024) // 12 mbps (bits)
-		, CommonMaxBitRate(12 * 1024 * 1024) // 12 mbps (bits)
+		, CommonMaxBitRate(16 * 1024 * 1024) // 16 mbps (bits)
 		, CommonQualityVsSpeed(100) // 0-33 Low Complexity, 34-66 Medium, 67-100 High (Higher = Slower but Better Quality)
 		, CommonConstantRateFactor(18) // 16-51, Higher is Worse
 		, EncodingProfile(EMoviePipelineMP4EncodeProfile::High)
 		, EncodingLevel(EMoviePipelineMP4EncodeLevel::Auto)
-		, EncodingRateControl(EMoviePipelineMP4EncodeRateControlMode::Quality)
+		, EncodingRateControl(EMoviePipelineMP4EncodeRateControlMode::VariableBitRate)
 	{}
 
 	/** The absolute path on disk to try and save the video file to. */
diff --git a/Engine/Plugins/MovieScene/MovieRenderPipeline/Source/MovieRenderPipelineMP4Encoder/Private/MovieRenderPipelineMP4EncoderModule.cpp b/Engine/Plugins/MovieScene/MovieRenderPipeline/Source/MovieRenderPipelineMP4Encoder/Private/MovieRenderPipelineMP4EncoderModule.cpp
--- a/Engine/Plugins/MovieScene/MovieRenderPipeline/Source/MovieRenderPipelineMP4Encoder/Private/MovieRenderPipelineMP4EncoderModule.cpp
+++ b/Engine/Plugins/MovieScene/MovieRenderPipeline/Source/MovieRenderPipelineMP4Encoder/Private/MovieRenderPipelineMP4EncoderModule.cpp
@@ -1,4 +1,36 @@
 // Copyright Epic Games, Inc. All Rights Reserved.
+
 #include "Modules/ModuleManager.h"
 
-IMPLEMENT_MODULE(FDefaultModuleImpl, MovieRenderPipelineMP4Encoder)
+#include "Graph/MovieGraphMP4EncoderNode.h"
+#include "Graph/MovieGraphMP4EncoderNodeCustomization.h"
+
+class FMovieRenderPipelineMP4EncoderModule : public IModuleInterface
+{
+public:
+	/*~ Begin IModuleInterface interface */
+	virtual void StartupModule() override;
+	virtual void ShutdownModule() override;
+	/*~ End IModuleInterface interface */
+};
+
+IMPLEMENT_MODULE(FMovieRenderPipelineMP4EncoderModule, MovieRenderPipelineMP4Encoder)
+
+void FMovieRenderPipelineMP4EncoderModule::StartupModule()
+{
+#if WITH_EDITOR
+	FPropertyEditorModule& PropertyModule = FModuleManager::LoadModuleChecked<FPropertyEditorModule>("PropertyEditor");
+	PropertyModule.RegisterCustomClassLayout(UMovieGraphMP4EncoderNode::StaticClass()->GetFName(), FOnGetDetailCustomizationInstance::CreateStatic(&FMovieGraphMP4EncoderNodeNodeCustomization::MakeInstance));
+#endif	// WITH_EDITOR
+}
+
+void FMovieRenderPipelineMP4EncoderModule::ShutdownModule()
+{
+#if WITH_EDITOR
+	// Unregister details customization
+	if (FPropertyEditorModule* PropertyModule = FModuleManager::GetModulePtr<FPropertyEditorModule>("PropertyEditor"))
+	{
+		PropertyModule->UnregisterCustomClassLayout(UMovieGraphMP4EncoderNode::StaticClass()->GetFName());
+	}
+#endif	// WITH_EDITOR
+}
\ No newline at end of file
diff --git a/Engine/Plugins/MovieScene/MovieRenderPipeline/Source/MovieRenderPipelineMP4Encoder/Private/Windows/MoviePipelineMP4Encoder.cpp b/Engine/Plugins/MovieScene/MovieRenderPipeline/Source/MovieRenderPipelineMP4Encoder/Private/Windows/MoviePipelineMP4Encoder.cpp
--- a/Engine/Plugins/MovieScene/MovieRenderPipeline/Source/MovieRenderPipelineMP4Encoder/Private/Windows/MoviePipelineMP4Encoder.cpp
+++ b/Engine/Plugins/MovieScene/MovieRenderPipeline/Source/MovieRenderPipelineMP4Encoder/Private/Windows/MoviePipelineMP4Encoder.cpp
@@ -342,7 +342,7 @@ static uint32_t GetEncodingRateControl(const EMoviePipelineMP4EncodeRateControlM
 	{
 	case EMoviePipelineMP4EncodeRateControlMode::ConstantBitRate: return eAVEncCommonRateControlMode_CBR;
 	case EMoviePipelineMP4EncodeRateControlMode::VariableBitRate_Constrained: return eAVEncCommonRateControlMode_PeakConstrainedVBR;
-	case EMoviePipelineMP4EncodeRateControlMode::VariableBitRate_Unconstrained: return eAVEncCommonRateControlMode_UnconstrainedVBR;
+	case EMoviePipelineMP4EncodeRateControlMode::VariableBitRate: return eAVEncCommonRateControlMode_UnconstrainedVBR;
 	case EMoviePipelineMP4EncodeRateControlMode::Quality: return eAVEncCommonRateControlMode_Quality;
 	case EMoviePipelineMP4EncodeRateControlMode::ConstantQP: return eAVEncCommonRateControlMode_PeakConstrainedVBR;
 	default:
@@ -665,7 +665,7 @@ static HRESULT GetVideoStreamEncoderAttributes(IMFAttributes** OutEncoderAttribu
 				Result = (*OutEncoderAttributes)->SetUINT32(CODECAPI_AVEncCommonMeanBitRate, InOptions.CommonMeanBitRate);
 			}
 		}
-		else if (InOptions.EncodingRateControl == EMoviePipelineMP4EncodeRateControlMode::VariableBitRate_Unconstrained)
+		else if (InOptions.EncodingRateControl == EMoviePipelineMP4EncodeRateControlMode::VariableBitRate)
 		{
 			if (SUCCEEDED(Result))
 			{
diff --git a/Engine/Plugins/MovieScene/MovieRenderPipeline/Source/MovieRenderPipelineMP4Encoder/Public/Graph/MovieGraphMP4EncoderNode.h b/Engine/Plugins/MovieScene/MovieRenderPipeline/Source/MovieRenderPipelineMP4Encoder/Public/Graph/MovieGraphMP4EncoderNode.h
--- a/Engine/Plugins/MovieScene/MovieRenderPipeline/Source/MovieRenderPipelineMP4Encoder/Public/Graph/MovieGraphMP4EncoderNode.h
+++ b/Engine/Plugins/MovieScene/MovieRenderPipeline/Source/MovieRenderPipelineMP4Encoder/Public/Graph/MovieGraphMP4EncoderNode.h
@@ -57,20 +57,20 @@ public:
 	UPROPERTY(EditAnywhere, BlueprintReadWrite, Category = Overrides, meta = (InlineEditConditionToggle))
 	uint8 bOverride_AverageBitrateInMbps : 1;
 	
-	UPROPERTY(EditAnywhere, BlueprintReadWrite, Category = Overrides, meta = (InlineEditConditionToggle))
-	uint8 bOverride_bIncludeAudio : 1;
-	
-	UPROPERTY(EditAnywhere, BlueprintReadWrite, Category = Overrides, meta = (InlineEditConditionToggle))
+	UPROPERTY(BlueprintReadWrite, Category = Overrides, meta = (InlineEditConditionToggle))
 	uint8 bOverride_MaxBitrateInMbps : 1;
 	
 	UPROPERTY(EditAnywhere, BlueprintReadWrite, Category = Overrides, meta = (InlineEditConditionToggle))
 	uint8 bOverride_ConstantRateFactor : 1;
-	
-	UPROPERTY(EditAnywhere, BlueprintReadWrite, Category = Overrides, meta = (InlineEditConditionToggle))
+
+	UPROPERTY(BlueprintReadWrite, Category = Overrides, meta = (InlineEditConditionToggle))
 	uint8 bOverride_EncodingProfile : 1;
 	
-	UPROPERTY(EditAnywhere, BlueprintReadWrite, Category = Overrides, meta = (InlineEditConditionToggle))
+	UPROPERTY(BlueprintReadWrite, Category = Overrides, meta = (InlineEditConditionToggle))
 	uint8 bOverride_EncodingLevel : 1;
+	
+	UPROPERTY(EditAnywhere, BlueprintReadWrite, Category = Overrides, meta = (InlineEditConditionToggle))
+	uint8 bOverride_bIncludeAudio : 1;
 
 	UPROPERTY(EditAnywhere, BlueprintReadWrite, Category = Overrides, meta = (InlineEditConditionToggle))
 	uint8 bOverride_OCIOConfiguration : 1;
@@ -80,15 +80,10 @@ public:
 		
 	/**
 	* Specifies the bitrate control method used by the encoder. Quality lets the user target a given quality without concern
-	* to the filesize, while Average/Constant modes allow you to suggest desired sizes, though the resulting file may still
-	* end up larger or smaller than specified.
+	* to the filesize, while Variable targets an average bit rate per frame.
 	*/
 	UPROPERTY(BlueprintReadWrite, EditAnywhere, Category = "H.264", meta=(EditCondition="bOverride_EncodingRateControl"))
-	EMoviePipelineMP4EncodeRateControlMode EncodingRateControl = EMoviePipelineMP4EncodeRateControlMode::Quality;
-
-	/** If true, audio will be included in the video file. */
-	UPROPERTY(EditAnywhere, BlueprintReadWrite, Category = "H.264", meta = (EditCondition = "bOverride_bIncludeAudio"))
-	bool bIncludeAudio = true;
+	EMoviePipelineMP4EncodeRateControlMode EncodingRateControl = EMoviePipelineMP4EncodeRateControlMode::VariableBitRate;
 	
 	/**
 	* What is the average bitrate the encoder should target per second? Value is in Megabits per Second,
@@ -106,8 +101,10 @@ public:
 	* maximum should be twice the average, but often in practice a smaller difference of 3-6Mbps is sufficient.
 	*
 	* Only applies to constrained max bit rate.
+	*
+	* Not exposed to the UI because it is expected most users do not need to change this, but is still available to be scripted.
 	*/
-	UPROPERTY(BlueprintReadWrite, EditAnywhere, Category = "H.264", meta = (ClampMin = 0.1, UIMin = 0.1, UIMax = 50, EditCondition = "bOverride_MaxBitrateInMbps"))
+	UPROPERTY(BlueprintReadWrite, Category = "H.264", meta = (ClampMin = 0.1, UIMin = 0.1, UIMax = 50, EditCondition = "bOverride_MaxBitrateInMbps"))
 	float MaxBitrateInMbps = 16.f;
 
 	/**
@@ -121,16 +118,26 @@ public:
 	UPROPERTY(BlueprintReadWrite, EditAnywhere, Category = "H.264", meta=(ClampMin=16, UIMin=16, ClampMax=51, UIMax=51, EditCondition="bOverride_ConstantRateFactor"))
 	int32 ConstantRateFactor = 20;
 
-	/** A higher profile generally results in a better quality video for the same bitrate, but may not be supported for playback on old devices. */
-	UPROPERTY(BlueprintReadWrite, EditAnywhere, Category = "H.264", meta=(AdvancedDisplay, EditCondition="bOverride_EncodingProfile"))
+	/**
+	 * A higher profile generally results in a better quality video for the same bitrate, but may not be supported for playback on old devices.
+	 *
+	 * Not exposed to the UI because it is expected most users do not need to change this, but is still available to be scripted.
+	 */
+	UPROPERTY(BlueprintReadWrite, Category = "H.264", meta=(AdvancedDisplay, EditCondition="bOverride_EncodingProfile"))
 	EMoviePipelineMP4EncodeProfile EncodingProfile = EMoviePipelineMP4EncodeProfile::High;
 
 	/** 
 	* A higher encode level generally results in a better quality video for the same bitrate, but may not be supported for playback on old devices.
 	* Auto will let the encoder automatically select the best level given the resolution, profile, and bitrate.
+	*
+	* Not exposed to the UI because it is expected most users do not need to change this, but is still available to be scripted.
 	*/
-	UPROPERTY(BlueprintReadWrite, EditAnywhere, Category = "H.264", meta=(AdvancedDisplay, EditCondition="bOverride_EncodingLevel"))
+	UPROPERTY(BlueprintReadWrite, Category = "H.264", meta=(AdvancedDisplay, EditCondition="bOverride_EncodingLevel"))
 	EMoviePipelineMP4EncodeLevel EncodingLevel = EMoviePipelineMP4EncodeLevel::Auto;
+
+	/** If true, audio will be included in the video file. */
+	UPROPERTY(EditAnywhere, BlueprintReadWrite, Category = "H.264", meta = (EditCondition = "bOverride_bIncludeAudio"))
+	bool bIncludeAudio = true;
 	
 	/**
 	* OCIO configuration/transform settings.
diff --git a/Engine/Plugins/MovieScene/MovieRenderPipeline/Source/MovieRenderPipelineMP4Encoder/Public/MoviePipelineMP4EncoderOutput.h b/Engine/Plugins/MovieScene/MovieRenderPipeline/Source/MovieRenderPipelineMP4Encoder/Public/MoviePipelineMP4EncoderOutput.h
--- a/Engine/Plugins/MovieScene/MovieRenderPipeline/Source/MovieRenderPipelineMP4Encoder/Public/MoviePipelineMP4EncoderOutput.h
+++ b/Engine/Plugins/MovieScene/MovieRenderPipeline/Source/MovieRenderPipelineMP4Encoder/Public/MoviePipelineMP4EncoderOutput.h
@@ -48,25 +48,23 @@ public:
 	UPROPERTY(BlueprintReadWrite, EditAnywhere, Category = "Settings")
 	EMoviePipelineMP4EncodeRateControlMode EncodingRateControl = EMoviePipelineMP4EncodeRateControlMode::Quality;
 	
-	/** If true, audio will be included in the video file. */
-	UPROPERTY(EditAnywhere, BlueprintReadWrite, Category = "Settings", meta = (EditCondition = "bOverride_bIncludeAudio"))
-	bool bIncludeAudio = true;
-	
 	/**
 	* What is the average bitrate the encoder should target per second? Value is in Megabits per Second,
 	* so a value of 8 will become 8192Kbps (kilobits per second). Higher resolutions and higher framerates
 	* need higher bitrates, reasonable starting values are 8 for 1080p30, 45 for 4k. Only applies to
 	* encoding modes not related to Quality.
 	*/
-	UPROPERTY(BlueprintReadWrite, EditAnywhere, Category = "Settings", meta=(ClampMin=0.1, UIMin=0.1, UIMax=50, EditCondition="EncodingRateControl != EMoviePipelineMP4EncodingRateControlMode::Quality"))
+	UPROPERTY(BlueprintReadWrite, EditAnywhere, Category = "Settings", meta=(ClampMin=0.1, UIMin=0.1, UIMax=50, EditCondition="EncodingRateControl == EMoviePipelineMP4EncodeRateControlMode::VariableBitRate", EditConditionHides))
 	float AverageBitrateInMbps = 8.f;
 	
 	/**
 	* When using VariableBitRate_Constrained, what is the maximum bitrate that the encoder can briefly use for
 	* more complex scenes, while still trying to maintain the average set in AverageBitrateInMbps. In theory the
 	* maximum should be twice the average, but often in practice a smaller difference of 3-6Mbps is sufficient.
+	*
+	* Not exposed to the UI because it is expected most users do not need to change this, but is still available to be scripted.
 	*/
-	UPROPERTY(BlueprintReadWrite, EditAnywhere, Category = "Settings", meta = (ClampMin = 0.1, UIMin = 0.1, UIMax = 50, EditCondition = "EncodingRateControl == EMoviePipelineMP4EncodingRateControlMode::VariableBitRate_Constrained"))
+	UPROPERTY(BlueprintReadWrite, Category = "Settings", meta = (ClampMin = 0.1, UIMin = 0.1, UIMax = 50))
 	float MaxBitrateInMbps = 16.f;
 
 	/**
@@ -75,20 +73,30 @@ public:
 	* maximum value range of [16-51], where 16 is the highest quality possible and 51 is the lowest quality. This scale
 	* is logarithmic, so small changes can result in large differences in quality, and filesize cost.
 	*/
-	UPROPERTY(BlueprintReadWrite, EditAnywhere, Category = "Settings", meta=(ClampMin=16, UIMin=16, ClampMax=51, UIMax=51, EditCondition="EncodingRateControl == EMoviePipelineMP4EncodingRateControlMode::Quality"))
+	UPROPERTY(BlueprintReadWrite, EditAnywhere, Category = "Settings", meta=(ClampMin=16, UIMin=16, ClampMax=51, UIMax=51, EditCondition="EncodingRateControl == EMoviePipelineMP4EncodeRateControlMode::Quality", EditConditionHides))
 	int32 ConstantRateFactor = 20;
 
-	/** A higher profile generally results in a better quality video for the same bitrate, but may not be supported for playback on old devices. */
-	UPROPERTY(BlueprintReadWrite, EditAnywhere, Category = "Settings", meta=(AdvancedDisplay))
+	/**
+	 * A higher profile generally results in a better quality video for the same bitrate, but may not be supported for playback on old devices.
+	 *
+	 * Not exposed to the UI because it is expected most users do not need to change this, but is still available to be scripted.
+	 */
+	UPROPERTY(BlueprintReadWrite, Category = "Settings", meta=(AdvancedDisplay))
 	EMoviePipelineMP4EncodeProfile EncodingProfile = EMoviePipelineMP4EncodeProfile::High;
 
 	/** 
 	* A higher encode level generally results in a better quality video for the same bitrate, but may not be supported for playback on old devices.
 	* Auto will let the encoder automatically select the best level given the resolution, profile, and bitrate.
+	*
+	* Not exposed to the UI because it is expected most users do not need to change this, but is still available to be scripted.
 	*/
-	UPROPERTY(BlueprintReadWrite, EditAnywhere, Category = "Settings", meta=(AdvancedDisplay))
+	UPROPERTY(BlueprintReadWrite, Category = "Settings", meta=(AdvancedDisplay))
 	EMoviePipelineMP4EncodeLevel EncodingLevel = EMoviePipelineMP4EncodeLevel::Auto;
 
+	/** If true, audio will be included in the video file. */
+	UPROPERTY(EditAnywhere, BlueprintReadWrite, Category = "Settings")
+	bool bIncludeAudio = true;
+
 protected:
 	struct FMP4CodecWriter : public MovieRenderPipeline::IVideoCodecWriter
 	{
