MoviePipeline: Fixed ASAN warning in MovieRenderPipelineMP4Encoder.

diff --git a/Engine/Plugins/MovieScene/MovieRenderPipeline/Source/MovieRenderPipelineMP4Encoder/Private/MovieRenderPipelineMP4EncoderModule.cpp b/Engine/Plugins/MovieScene/MovieRenderPipeline/Source/MovieRenderPipelineMP4Encoder/Private/MovieRenderPipelineMP4EncoderModule.cpp
--- a/Engine/Plugins/MovieScene/MovieRenderPipeline/Source/MovieRenderPipelineMP4Encoder/Private/MovieRenderPipelineMP4EncoderModule.cpp
+++ b/Engine/Plugins/MovieScene/MovieRenderPipeline/Source/MovieRenderPipelineMP4Encoder/Private/MovieRenderPipelineMP4EncoderModule.cpp
@@ -12,6 +12,10 @@ public:
 	virtual void StartupModule() override;
 	virtual void ShutdownModule() override;
 	/*~ End IModuleInterface interface */
+
+private:
+	/** StaticClass() isn't safe to call during ShutdownModule(), so cache the names of registered classes here. */
+	TArray<FName> ClassesToUnregisterOnShutdown;
 };
 
 IMPLEMENT_MODULE(FMovieRenderPipelineMP4EncoderModule, MovieRenderPipelineMP4Encoder)
@@ -20,17 +24,22 @@ void FMovieRenderPipelineMP4EncoderModule::StartupModule()
 {
 #if WITH_EDITOR
 	FPropertyEditorModule& PropertyModule = FModuleManager::LoadModuleChecked<FPropertyEditorModule>("PropertyEditor");
-	PropertyModule.RegisterCustomClassLayout(UMovieGraphMP4EncoderNode::StaticClass()->GetFName(), FOnGetDetailCustomizationInstance::CreateStatic(&FMovieGraphMP4EncoderNodeNodeCustomization::MakeInstance));
+
+	ClassesToUnregisterOnShutdown.Add(UMovieGraphMP4EncoderNode::StaticClass()->GetFName());
+	PropertyModule.RegisterCustomClassLayout(ClassesToUnregisterOnShutdown.Last(), FOnGetDetailCustomizationInstance::CreateStatic(&FMovieGraphMP4EncoderNodeNodeCustomization::MakeInstance));
 #endif	// WITH_EDITOR
 }
 
 void FMovieRenderPipelineMP4EncoderModule::ShutdownModule()
 {
 #if WITH_EDITOR
-	// Unregister details customization
+	// Unregister details customizations
 	if (FPropertyEditorModule* PropertyModule = FModuleManager::GetModulePtr<FPropertyEditorModule>("PropertyEditor"))
 	{
-		PropertyModule->UnregisterCustomClassLayout(UMovieGraphMP4EncoderNode::StaticClass()->GetFName());
+		for (const FName& ClassToUnregister : ClassesToUnregisterOnShutdown)
+		{
+			PropertyModule->UnregisterCustomClassLayout(ClassToUnregister);
+		}
 	}
 #endif	// WITH_EDITOR
 }
\ No newline at end of file
