MoviePipeline: Fixed issue with MP4 videos encoding upside-down with the software encoder and AMD hardware encoder.

diff --git a/Engine/Plugins/MovieScene/MovieRenderPipeline/Source/MovieRenderPipelineMP4Encoder/Private/Graph/MovieGraphMP4EncoderNode.cpp b/Engine/Plugins/MovieScene/MovieRenderPipeline/Source/MovieRenderPipelineMP4Encoder/Private/Graph/MovieGraphMP4EncoderNode.cpp
--- a/Engine/Plugins/MovieScene/MovieRenderPipeline/Source/MovieRenderPipelineMP4Encoder/Private/Graph/MovieGraphMP4EncoderNode.cpp
+++ b/Engine/Plugins/MovieScene/MovieRenderPipeline/Source/MovieRenderPipelineMP4Encoder/Private/Graph/MovieGraphMP4EncoderNode.cpp
@@ -107,7 +107,7 @@ bool UMovieGraphMP4EncoderNode::Initialize_EncodeThread(MovieRenderGraph::IVideo
 	const FMP4CodecWriter* CodecWriter = static_cast<FMP4CodecWriter*>(InWriter);
 	if(!CodecWriter->Writer->Initialize())
 	{
-		UE_LOG(LogMovieRenderPipeline, Error, TEXT("Failed to initialize Movie Pipeline MP4 writer. This can sometimes be caused by trying to encode to a resolution higher than 4K (4096), which is unsupported."));
+		UE_LOG(LogMovieRenderPipeline, Error, TEXT("Failed to initialize Movie Pipeline MP4 writer. An encoder that supports the render resolution and requested MP4 encode options was not found. Try again with a different resolution and/or encoder settings."));
 		return false;
 	}
 	
diff --git a/Engine/Plugins/MovieScene/MovieRenderPipeline/Source/MovieRenderPipelineMP4Encoder/Private/MoviePipelineMP4EncoderOutput.cpp b/Engine/Plugins/MovieScene/MovieRenderPipeline/Source/MovieRenderPipelineMP4Encoder/Private/MoviePipelineMP4EncoderOutput.cpp
--- a/Engine/Plugins/MovieScene/MovieRenderPipeline/Source/MovieRenderPipelineMP4Encoder/Private/MoviePipelineMP4EncoderOutput.cpp
+++ b/Engine/Plugins/MovieScene/MovieRenderPipeline/Source/MovieRenderPipelineMP4Encoder/Private/MoviePipelineMP4EncoderOutput.cpp
@@ -49,7 +49,7 @@ bool UMoviePipelineMP4EncoderOutput::Initialize_EncodeThread(MovieRenderPipeline
 	FMP4CodecWriter* CodecWriter = static_cast<FMP4CodecWriter*>(InWriter);
 	if(!CodecWriter->Writer->Initialize())
 	{
-		UE_LOG(LogMovieRenderPipeline, Error, TEXT("Failed to initialize Media Foundation Sink Writer."));
+		UE_LOG(LogMovieRenderPipeline, Error, TEXT("Failed to initialize Movie Pipeline MP4 writer. An encoder that supports the render resolution and requested MP4 encode options was not found. Try again with a different resolution and/or encoder settings."));
 		return false;
 	}
 
diff --git a/Engine/Plugins/MovieScene/MovieRenderPipeline/Source/MovieRenderPipelineMP4Encoder/Private/Windows/MoviePipelineMP4Encoder.cpp b/Engine/Plugins/MovieScene/MovieRenderPipeline/Source/MovieRenderPipelineMP4Encoder/Private/Windows/MoviePipelineMP4Encoder.cpp
--- a/Engine/Plugins/MovieScene/MovieRenderPipeline/Source/MovieRenderPipelineMP4Encoder/Private/Windows/MoviePipelineMP4Encoder.cpp
+++ b/Engine/Plugins/MovieScene/MovieRenderPipeline/Source/MovieRenderPipelineMP4Encoder/Private/Windows/MoviePipelineMP4Encoder.cpp
@@ -2,6 +2,8 @@
 
 #include "MoviePipelineMP4Encoder.h"
 
+#include "MovieRenderPipelineCoreModule.h"
+
 #if PLATFORM_WINDOWS
 #include "Windows/AllowWindowsPlatformTypes.h"
 #include "Windows/WindowsPlatformMisc.h"
@@ -22,6 +24,39 @@ THIRD_PARTY_INCLUDES_END
 #include "Windows/HideWindowsPlatformTypes.h"
 #endif
 
+namespace
+{
+	/** Gets the path of the DLL that the ICodecAPI comes from. */
+	FString GetCodecApiDllPath(ICodecAPI* InCodecAPI)
+	{
+		static const FString UnknownDLL(TEXT("Unknown"));
+
+		if (!InCodecAPI)
+		{
+			return UnknownDLL;
+		}
+
+		// This is really funky pointer trickery. According to Microsoft, "A COM interface pointer is a pointer to a structure that consists of just
+		// a vtable". The first entry in the vtable is the address to the QueryInterface (which will originate from the DLL of the ICodecAPI
+		// implementation that is being used).
+		void** VtableAddress = *reinterpret_cast<void***>(InCodecAPI);
+		const void* QueryInterfaceAddress = VtableAddress[0];
+
+		// Get memory information about the page at the query interface's address
+		MEMORY_BASIC_INFORMATION MemoryInformation;
+		if (VirtualQueryEx(GetCurrentProcess(), QueryInterfaceAddress, &MemoryInformation, sizeof(MemoryInformation)))
+		{
+			// Get the path of the module handle
+			WCHAR DllPath[MAX_PATH] = {};
+			if (GetModuleFileNameW(static_cast<HMODULE>(MemoryInformation.AllocationBase), DllPath, MAX_PATH))
+			{
+				return FString(DllPath);
+			}
+		}
+
+		return UnknownDLL;
+	}
+}
 
 FMoviePipelineMP4Encoder::FMoviePipelineMP4Encoder(const FMoviePipelineMP4EncoderOptions& InOptions)
 	: Options(InOptions)
@@ -417,10 +452,7 @@ static HRESULT CreateVideoMediaOutputStream(IMFMediaType** OutVideoMediaTypeOutp
 
 	if (SUCCEEDED(Result))
 	{
-		// Setting this to anything other than Unknown seems to flip the resulting output video, but only on
-		// hardware encoders, but not software encoder. But if the software encoder gets used, then the colorspace
-		// is wrong.
-		Result = (*OutVideoMediaTypeOutput)->SetUINT32(MF_MT_VIDEO_NOMINAL_RANGE, MFNominalRange_Unknown);
+		Result = (*OutVideoMediaTypeOutput)->SetUINT32(MF_MT_VIDEO_NOMINAL_RANGE, MFNominalRange_16_235);
 	}
 
 	if (SUCCEEDED(Result))
@@ -566,6 +598,19 @@ static HRESULT CreateVideoMediaTypeIn(IMFMediaType** OutVideoMediaTypeIn, const
 		Result = (*OutVideoMediaTypeIn)->SetUINT32(MF_MT_ALL_SAMPLES_INDEPENDENT, true);
 	}
 
+	// Setting an explicit (positive) stride and a sample size is important to ensure all encoders process frames top-down, otherwise some encoders
+	// may process frames upside-down.
+	const uint32 Stride = InOptions.Width * 4;	// 4 bytes per pixel
+	if (SUCCEEDED(Result))
+	{
+		Result = (*OutVideoMediaTypeIn)->SetUINT32(MF_MT_DEFAULT_STRIDE, Stride);
+	}
+
+	if (SUCCEEDED(Result))
+	{
+		Result = (*OutVideoMediaTypeIn)->SetUINT32(MF_MT_SAMPLE_SIZE, Stride * InOptions.Height);
+	}
+
 	return Result;
 }
 
@@ -748,6 +793,7 @@ bool FMoviePipelineMP4Encoder::InitializeEncoder()
 	{
 		Result = CreateVideoMediaOutputStream(&VideoMediaTypeOut, Options);
 	}
+
 	// Video Stream
 	if (SUCCEEDED(Result) && SinkWriter)
 	{
@@ -801,13 +847,17 @@ bool FMoviePipelineMP4Encoder::InitializeEncoder()
 		}
 	}
 
-	// Note: This has no practical use to the user. However, for diagnostic purposes, `codecApi` can be inspected in the debugger to determine which
-	// encoder is being used (software or hardware).
-	// if (SUCCEEDED(Result) && SinkWriter)
-	// {
-	// 	ICodecAPI* codecApi = nullptr;
-	// 	Result = SinkWriter->GetServiceForStream(VideoStreamIndex, GUID_NULL, __uuidof(ICodecAPI), (LPVOID*)&codecApi);
-	// }
+	// Log the path of the DLL the encoder is using for debug purposes
+	if (SUCCEEDED(Result) && SinkWriter)
+	{
+		ICodecAPI* CodecApi = nullptr;
+		Result = SinkWriter->GetServiceForStream(VideoStreamIndex, GUID_NULL, __uuidof(ICodecAPI), (LPVOID*)&CodecApi);
+		
+		// The best we can do is provide the DLL path. Ideally we could use IMFSinkWriterEx to get the friendly name attribute for the active
+		// encoder, but IMFSinkWriterEx is not available with the currently defined WINVER (needs to be >= 0x0602).
+		const FString CodecApiDllPath = GetCodecApiDllPath(CodecApi);
+		UE_LOG(LogMovieRenderPipeline, Display, TEXT("Using the following encoder for the MP4 encode: %s"), *CodecApiDllPath);
+	}
 
 	// The sink writer can now accept data.
 	if (SUCCEEDED(Result) && SinkWriter)
