MovieGraph: Implement Page to System Memory on Deferred Panoramic Renders and Deferred High Res Tiling. Page to System Memory, when using history-per-tile (or pane), will download the per-tile GPU resources off of the GPU after each tile is rendered, and then re-upload them on the next frame that needs it. This is about 3x slower than without it, but allows for using large amounts of system memory to free up some GPU memory if history is required for the scene to render correctly. High resolution images/panoramics are still targeted at high-VRAM cards regardless of this feature.

diff --git a/Engine/Plugins/MovieScene/MovieRenderPipeline/Source/MovieRenderPipelineCore/Private/Graph/MovieGraphDefaultRenderer.cpp b/Engine/Plugins/MovieScene/MovieRenderPipeline/Source/MovieRenderPipelineCore/Private/Graph/MovieGraphDefaultRenderer.cpp
--- a/Engine/Plugins/MovieScene/MovieRenderPipeline/Source/MovieRenderPipelineCore/Private/Graph/MovieGraphDefaultRenderer.cpp
+++ b/Engine/Plugins/MovieScene/MovieRenderPipeline/Source/MovieRenderPipelineCore/Private/Graph/MovieGraphDefaultRenderer.cpp
@@ -333,9 +333,17 @@ void UMovieGraphDefaultRenderer::Render(const FMovieGraphTimeStepData& InTimeSte
 #if WITH_EDITOR && !UE_BUILD_SHIPPING
 	TUniquePtr<RenderCaptureInterface::FScopedCapture> ScopedGPUCapture;
 	{
-		UMovieGraphDebugSettingNode* DebugSettings = InTimeStepData.EvaluatedConfig->GetSettingForBranch<UMovieGraphDebugSettingNode>(UMovieGraphNode::GlobalsPinName);
 		bool bIsRendering = GetOwningGraph()->GetActiveShotList()[GetOwningGraph()->GetCurrentShotIndex()]->ShotInfo.State == EMovieRenderShotState::Rendering;
-		if (DebugSettings && DebugSettings->bCaptureFramesWithRenderDoc && bIsRendering)
+		UMovieGraphDebugSettingNode* DebugSettings = InTimeStepData.EvaluatedConfig->GetSettingForBranch<UMovieGraphDebugSettingNode>(UMovieGraphNode::GlobalsPinName);
+
+		const bool bWantsDebugCapture = DebugSettings && DebugSettings->bCaptureFramesWithRenderDoc;
+		bool bIsCorrectFrame = false;
+		if (bWantsDebugCapture)
+		{
+			bIsCorrectFrame = InTimeStepData.RootFrameNumber == DebugSettings->RenderDocCaptureFrame;
+		}
+
+		if (bIsRendering && bWantsDebugCapture && bIsCorrectFrame)
 		{
 			ScopedGPUCapture = MakeUnique<RenderCaptureInterface::FScopedCapture>(true, *FString::Printf(TEXT("MRQ Frame: %d"), InTimeStepData.RootFrameNumber.Value));
 		}
diff --git a/Engine/Plugins/MovieScene/MovieRenderPipeline/Source/MovieRenderPipelineCore/Public/Graph/Nodes/MovieGraphDebugNode.h b/Engine/Plugins/MovieScene/MovieRenderPipeline/Source/MovieRenderPipelineCore/Public/Graph/Nodes/MovieGraphDebugNode.h
--- a/Engine/Plugins/MovieScene/MovieRenderPipeline/Source/MovieRenderPipelineCore/Public/Graph/Nodes/MovieGraphDebugNode.h
+++ b/Engine/Plugins/MovieScene/MovieRenderPipeline/Source/MovieRenderPipelineCore/Public/Graph/Nodes/MovieGraphDebugNode.h
@@ -26,6 +26,9 @@ public:
 public:
 	UPROPERTY(EditAnywhere, BlueprintReadWrite, Category = Overrides, meta = (InlineEditConditionToggle))
 	uint8 bOverride_bCaptureFramesWithRenderDoc : 1;
+
+	UPROPERTY(EditAnywhere, BlueprintReadWrite, Category = Overrides, meta = (InlineEditConditionToggle))
+	uint8 bOverride_RenderDocCaptureFrame : 1;
 	
 	UPROPERTY(EditAnywhere, BlueprintReadWrite, Category = Overrides, meta = (InlineEditConditionToggle))
 	uint8 bOverride_bCaptureUnrealInsightsTrace : 1;
@@ -39,6 +42,12 @@ public:
 	*/
 	UPROPERTY(EditAnywhere, BlueprintReadWrite, Category = "Settings", meta = (EditCondition = "bOverride_bCaptureFramesWithRenderDoc"))
 	bool bCaptureFramesWithRenderDoc;
+
+	/**
+	* If bCaptureFramesWithRenderDoc is true, which frame (on the root Sequencer time line) should we capture data for?
+	*/
+	UPROPERTY(EditAnywhere, BlueprintReadWrite, Category = "Settings", meta = (EditCondition = "bOverride_RenderDocCaptureFrame"))
+	int32 RenderDocCaptureFrame;
 	
 	/** 
 	* If true, automatically capture an Unreal Insights trace file for the duration of the render.
diff --git a/Engine/Plugins/MovieScene/MovieRenderPipeline/Source/MovieRenderPipelineRenderPasses/Private/Graph/Renderers/MovieGraphDeferredPanoramicPass.cpp b/Engine/Plugins/MovieScene/MovieRenderPipeline/Source/MovieRenderPipelineRenderPasses/Private/Graph/Renderers/MovieGraphDeferredPanoramicPass.cpp
--- a/Engine/Plugins/MovieScene/MovieRenderPipeline/Source/MovieRenderPipelineRenderPasses/Private/Graph/Renderers/MovieGraphDeferredPanoramicPass.cpp
+++ b/Engine/Plugins/MovieScene/MovieRenderPipeline/Source/MovieRenderPipelineRenderPasses/Private/Graph/Renderers/MovieGraphDeferredPanoramicPass.cpp
@@ -272,8 +272,8 @@ void FMovieGraphDeferredPanoramicPass::Setup(TWeakObjectPtr<UMovieGraphDefaultRe
 	LayerData = InLayer;
 
 	UMovieGraphDeferredPanoramicNode* ParentNode = Cast<UMovieGraphDeferredPanoramicNode>(InLayer.RenderPassNode);
-	int32 NumHorizontalSteps = ParentNode->NumHorizontalSteps;
-	int32 NumVerticalSteps = ParentNode->NumVerticalSteps;
+	int32 NumHorizontalSteps = FMath::Max(0, ParentNode->NumHorizontalSteps);
+	int32 NumVerticalSteps = FMath::Max(0, ParentNode->NumVerticalSteps);
 	int32 NumPanoramicPanes = NumHorizontalSteps * NumVerticalSteps;
 
 	if (ParentNode->bAllocateHistoryPerPane)
@@ -293,6 +293,18 @@ void FMovieGraphDeferredPanoramicPass::Setup(TWeakObjectPtr<UMovieGraphDefaultRe
 		if (bUsesAutoExposure)
 		{
 			AutoExposureViewStates.SetNum(CubeFace_MAX);
+
+			// ShareOrigin must be called before Allocate.  ShareOrigin is necessary for Lumen to work with 6 way cube split screen
+			// (causes Lumen scene data to be shared for all views, and overrides Lumen's regular 2 view split screen limitation).
+			for (int32 Index = 1; Index < CubeFace_MAX; Index++)
+			{
+				AutoExposureViewStates[Index].ShareOrigin(&AutoExposureViewStates[0]);
+			}
+
+			for (int32 Index = 0; Index < CubeFace_MAX; Index++)
+			{
+				AutoExposureViewStates[Index].Allocate(InRenderer->GetWorld()->GetFeatureLevel());
+			}
 		}
 
 		// Now that we've stopped allocating View States, we can Allocate them all.
@@ -301,17 +313,6 @@ void FMovieGraphDeferredPanoramicPass::Setup(TWeakObjectPtr<UMovieGraphDefaultRe
 			PaneViewStates[Index].Allocate(InRenderer->GetWorld()->GetFeatureLevel());
 		}
 
-		// ShareOrigin must be called before Allocate.  ShareOrigin is necessary for Lumen to work with 6 way cube split screen
-		// (causes Lumen scene data to be shared for all views, and overrides Lumen's regular 2 view split screen limitation).
-		for (int32 Index = 1; Index < CubeFace_MAX; Index++)
-		{
-			AutoExposureViewStates[Index].ShareOrigin(&AutoExposureViewStates[0]);
-		}
-
-		for (int32 Index = 0; Index < CubeFace_MAX; Index++)
-		{
-			AutoExposureViewStates[Index].Allocate(InRenderer->GetWorld()->GetFeatureLevel());
-		}
 	}
 }
 
@@ -475,8 +476,6 @@ void FMovieGraphDeferredPanoramicPass::Render(const FMovieGraphTraversalContext&
 	const ESceneCaptureSource SceneCaptureSource = ParentNodeThisFrame->GetDisableToneCurve() ? ESceneCaptureSource::SCS_FinalColorHDR : ESceneCaptureSource::SCS_FinalToneCurveHDR;
 	const EAntiAliasingMethod AntiAliasingMethod = UE::MovieRenderPipeline::GetEffectiveAntiAliasingMethod(ParentNodeThisFrame->GetOverrideAntiAliasing(), ParentNodeThisFrame->GetAntiAliasingMethod());
 
-
-
 	if (!bHasPrintedRenderingInfo)
 	{
 		UE_LOG(LogMovieRenderPipeline, Log, TEXT("Set-up Deferred Panoramic Renderer: %s Layer: %s OutputRes: [%d, %d] PaneRes: [%d, %d] PaneCount: [%d, %d] bPageToSystemMemory: %d bAutoExposurePass: %d"),
@@ -543,13 +542,27 @@ void FMovieGraphDeferredPanoramicPass::Render(const FMovieGraphTraversalContext&
 				// World should be paused for every spatial sample except the last one, so that
 				// the view doesn't update histories until the end, allowing us to render the same
 				// scene multiple times.
-				const bool bIsLastPane = FIntPoint(PaneX, PaneY) == FIntPoint(PaneCount.X - 1, PaneCount.Y - 1);
-				const bool bWorldIsPaused = !(bIsLastPane && SpatialIndex == (NumSpatialSamples - 1));
+				const bool bWorldIsPaused = !(SpatialIndex == (NumSpatialSamples - 1));
 				const int32 FrameIndex = InTimeData.RenderedFrameNumber * ((InTimeData.TemporalSampleCount * NumSpatialSamples) + (InTimeData.TemporalSampleIndex * NumSpatialSamples)) + SpatialIndex;
 
+				// We need to do this check before we start seeing if we need anti-aliasing samples so that when it falls back
+				// to no AA, it still does the right thing and produces AA if they have spatial/temporal samples.
+				EAntiAliasingMethod EffectiveAntiAliasingMethod = AntiAliasingMethod;
+				const bool bRequiresHistory = IsTemporalAccumulationBasedMethod(EffectiveAntiAliasingMethod);
+				if (!ParentNodeThisFrame->bAllocateHistoryPerPane && bRequiresHistory)
+				{
+					if (!bHasPrintedWarnings)
+					{
+						UE_LOG(LogMovieRenderPipeline, Warning, TEXT("Panoramic Renders do not support TAA/TSR without enabling bAllocateHistoryPerPane! Forcing AntiAliasing off."));
+						bHasPrintedWarnings = true;
+					}
+
+					EffectiveAntiAliasingMethod = EAntiAliasingMethod::AAM_None;
+				}
+
 				// We only allow a spatial jitter if we have more than one sample
 				FVector2f SpatialShiftAmount = FVector2f(0.f, 0.f);
-				const bool bAntiAliasingAllowsJitter = AntiAliasingMethod == EAntiAliasingMethod::AAM_None;
+				const bool bAntiAliasingAllowsJitter = EffectiveAntiAliasingMethod == EAntiAliasingMethod::AAM_None;
 				const bool bSampleCountsAllowsJitter = NumSpatialSamples > 1 || InTimeData.TemporalSampleCount > 1;
 				if (bAntiAliasingAllowsJitter && bSampleCountsAllowsJitter)
 				{
@@ -587,13 +600,13 @@ void FMovieGraphDeferredPanoramicPass::Render(const FMovieGraphTraversalContext&
 				ViewFamilyInitData.SceneCaptureSource = SceneCaptureSource;
 				ViewFamilyInitData.bWorldIsPaused = bWorldIsPaused;
 				ViewFamilyInitData.FrameIndex = FrameIndex;
-				ViewFamilyInitData.AntiAliasingMethod = AntiAliasingMethod;
+				ViewFamilyInitData.AntiAliasingMethod = EffectiveAntiAliasingMethod;
 				ViewFamilyInitData.ShowFlags = ParentNodeThisFrame->GetShowFlags();
 				ViewFamilyInitData.ViewModeIndex = ParentNodeThisFrame->GetViewModeIndex();
 				ViewFamilyInitData.ProjectionMode = CameraInfo.ViewInfo.ProjectionMode;
 
 				FSceneViewStateInterface* ViewState = GetSceneViewState(ParentNodeThisFrame, PaneX, PaneY);
-				if (ParentNodeThisFrame->bAllocateHistoryPerPane && ParentNodeThisFrame->bPageToSystemMemory)
+				if (ViewState && ParentNodeThisFrame->bAllocateHistoryPerPane && ParentNodeThisFrame->bPageToSystemMemory)
 				{
 					// If paging to system memory, restore the data needed for this particular Scene View History,
 					// Transfering from CPU->GPU
@@ -665,7 +678,7 @@ void FMovieGraphDeferredPanoramicPass::Render(const FMovieGraphTraversalContext&
 					SampleState.bFetchFromAccumulator = InTimeData.bIsLastTemporalSampleForFrame && (SpatialIndex == (NumSpatialSamples - 1));
 					SampleState.OverlappedPad = FIntPoint(0, 0);
 					SampleState.OverlappedOffset = FIntPoint(0, 0);
-					SampleState.OverlappedSubpixelShift = FVector2D(0.5, 0.5);
+					SampleState.OverlappedSubpixelShift = FVector2D(0.5f - SpatialShiftAmount.X, 0.5f - SpatialShiftAmount.Y);
 					SampleState.OverscanFraction = 0.f;
 					SampleState.CropRectangle = FIntRect(0, 0, DesiredOutputResolution.X, DesiredOutputResolution.Y); // ToDo: Output resolution will be forced to a 2:1 ratio but this currently respects what the user put in.
 					SampleState.bAllowOCIO = ParentNodeThisFrame->GetAllowOCIO();
@@ -683,19 +696,6 @@ void FMovieGraphDeferredPanoramicPass::Render(const FMovieGraphTraversalContext&
 
 				ApplyMovieGraphOverridesToSampleState(SampleState);
 
-				const bool bRequiresHistory = IsTemporalAccumulationBasedMethod(ViewFamily->Views[0]->AntiAliasingMethod);
-				if (!ParentNodeThisFrame->bAllocateHistoryPerPane && bRequiresHistory)
-				{
-					if (!bHasPrintedWarnings)
-					{
-						UE_LOG(LogMovieRenderPipeline, Warning, TEXT("Panoramic Renders do not support TAA/TSR without enabling bAllocateHistoryPerPane! Forcing AntiAliasing off."));
-						bHasPrintedWarnings = true;
-					}
-
-					FSceneView* NonConstView = const_cast<FSceneView*>(ViewFamily->Views[0]);
-					NonConstView->AntiAliasingMethod = EAntiAliasingMethod::AAM_None;
-				}
-
 				// If this was just to contribute to the history buffer, no need to go any further. Never discard if we're writing individual samples, though.
 				bool bDiscardOutput = (InTimeData.bDiscardOutput || ShouldDiscardOutput(ViewFamily, CameraInfo)) && !SampleState.bWriteSampleToDisk;
 
@@ -709,6 +709,12 @@ void FMovieGraphDeferredPanoramicPass::Render(const FMovieGraphTraversalContext&
 						RHICmdList.Transition(FRHITransitionInfo(RenderTargetResource->GetRenderTargetTexture(), ERHIAccess::RTV, ERHIAccess::SRVGraphicsPixel));
 					});
 
+				// After submission, if we're paging to system memory, mark the resources for download into system memory.
+				if (ViewState && ParentNodeThisFrame->bAllocateHistoryPerPane && ParentNodeThisFrame->bPageToSystemMemory)
+				{
+					ViewState->SystemMemoryMirrorBackup(SystemMemoryMirror.Get());
+				}
+
 				// Data may not be something we actually want to read back and write to disk after rendering.
 				if (bDiscardOutput)
 				{
@@ -716,12 +722,6 @@ void FMovieGraphDeferredPanoramicPass::Render(const FMovieGraphTraversalContext&
 				}
 
 				PostRendererSubmission(SampleState, RenderTargetInitParams, Canvas, CameraInfo);
-
-				// After submission, if we're paging to system memory, mark the resources for download into system memory.
-				if (ParentNodeThisFrame->bAllocateHistoryPerPane && ParentNodeThisFrame->bPageToSystemMemory)
-				{
-					ViewState->SystemMemoryMirrorBackup(SystemMemoryMirror.Get());
-				}
 			}
 		}
 	}
diff --git a/Engine/Plugins/MovieScene/MovieRenderPipeline/Source/MovieRenderPipelineRenderPasses/Private/Graph/Renderers/MovieGraphDeferredPass.cpp b/Engine/Plugins/MovieScene/MovieRenderPipeline/Source/MovieRenderPipelineRenderPasses/Private/Graph/Renderers/MovieGraphDeferredPass.cpp
--- a/Engine/Plugins/MovieScene/MovieRenderPipeline/Source/MovieRenderPipelineRenderPasses/Private/Graph/Renderers/MovieGraphDeferredPass.cpp
+++ b/Engine/Plugins/MovieScene/MovieRenderPipeline/Source/MovieRenderPipelineRenderPasses/Private/Graph/Renderers/MovieGraphDeferredPass.cpp
@@ -89,20 +89,31 @@ void FMovieGraphDeferredPass::Setup(TWeakObjectPtr<UMovieGraphDefaultRenderer> I
 	SceneViewState.Allocate(InRenderer->GetWorld()->GetFeatureLevel());
 	PRAGMA_ENABLE_DEPRECATION_WARNINGS
 
-	// Auto Exposure passes are only needed if:
-	// A) The user has requested them
-	// B) High Resolution Tiling is enabled
-	// C) The tile count needed by high resolution is > 1.
-	const bool bCameraUsesAutoExposure = false; // ToDo: Auto-exposure support.
-	const bool bHighRes = ParentNode->GetEnableHighResolutionTiling();
-	const bool bHasTiles = (TileCountWithHistory.X * TileCountWithHistory.Y) > 1;
-
-	bHasAutoExposurePass = bCameraUsesAutoExposure && bHighRes && bHasTiles;
-	if (bHasAutoExposurePass)
+	if (ParentNode->GetEnableHistoryPerTile())
 	{
-		// Auto Exposure passes are always at -1, -1
-		SceneViewStates.Add(FIntPoint(-1, -1));
-		NumSceneViewStateReferencesAllocated++;
+		if (ParentNode->GetEnablePageToSystemMemory())
+		{
+			SystemMemoryMirror = FSceneViewStateInterface::SystemMemoryMirrorAllocate();
+		}
+
+		// ToDo: This doesn't take into account blended post-process values from the world, but we don't have a way to do the blending without having a FSceneView which doesn't exist until render time.
+		const UE::MovieGraph::DefaultRenderer::FCameraInfo& CameraInfo = InRenderer->GetCameraInfo(InLayer.CameraIndex);
+		const bool bUsesAutoExposure = CameraInfo.ViewInfo.PostProcessSettings.bOverride_AutoExposureMethod && CameraInfo.ViewInfo.PostProcessSettings.AutoExposureMethod != EAutoExposureMethod::AEM_Manual;
+
+		const bool bHighRes = ParentNode->GetEnableHighResolutionTiling();
+		const bool bHasTiles = (TileCountWithHistory.X * TileCountWithHistory.Y) > 1;
+
+		// Auto Exposure passes are only needed if:
+		// A) The user has requested them
+		// B) High Resolution Tiling is enabled
+		// C) The tile count needed by high resolution is > 1.
+		bHasAutoExposurePass = bUsesAutoExposure && bHighRes && bHasTiles;
+		if (bHasAutoExposurePass)
+		{
+			// Auto Exposure passes are always at -1, -1
+			SceneViewStates.Add(FIntPoint(-1, -1));
+			NumSceneViewStateReferencesAllocated++;
+		}
 	}
 
 	for (TPair<FIntPoint, FSceneViewStateReference>& Pair : SceneViewStates)
@@ -120,7 +131,7 @@ void FMovieGraphDeferredPass::Setup(TWeakObjectPtr<UMovieGraphDefaultRenderer> I
 
 	UE_LOG(LogMovieRenderPipeline, Log, TEXT("Initialized Renderer: %s Layer: %s TileCount[%d, %d]/[%d, %d] AutoExposurePass: %d HighRes: %d PageToSystem: %d TotalSceneViewStateReferences: %d"),
 		*RenderDataIdentifier.RendererName, *RenderDataIdentifier.LayerName, ParentNode->GetTileCount().X, ParentNode->GetTileCount().Y, TileCountWithHistory.X, TileCountWithHistory.Y,
-		bHasAutoExposurePass, bHighRes, ParentNode->GetEnablePageToSystemMemory(), NumSceneViewStateReferencesAllocated);
+		bHasAutoExposurePass, ParentNode->GetEnableHighResolutionTiling(), ParentNode->GetEnablePageToSystemMemory(), NumSceneViewStateReferencesAllocated);
 }
 
 void FMovieGraphDeferredPass::Teardown()
@@ -383,9 +394,16 @@ void FMovieGraphDeferredPass::Render(const FMovieGraphTraversalContext& InFrameT
 				ViewFamilyInitData.ViewModeIndex = ParentNodeThisFrame->GetViewModeIndex();
 				ViewFamilyInitData.ProjectionMode = CameraInfo.ViewInfo.ProjectionMode;
 
+				FSceneViewStateInterface* ViewState = GetSceneViewState(ParentNodeThisFrame, TileX, TileY);
+				
+				// Optional system memory mirroring
+				if (ViewState && ParentNodeThisFrame->GetEnableHistoryPerTile() && ParentNodeThisFrame->GetEnablePageToSystemMemory())
+				{
+					ViewState->SystemMemoryMirrorRestore(SystemMemoryMirror.Get());
+				}
+
 				TSharedRef<FSceneViewFamilyContext> ViewFamily = CreateSceneViewFamily(ViewFamilyInitData);
 
-				FSceneViewStateInterface* ViewState = GetSceneViewState(ParentNodeThisFrame, TileX, TileY);
 				FSceneViewInitOptions SceneViewInitOptions = CreateViewInitOptions(CameraInfo, ViewFamily.ToSharedPtr().Get(), ViewState);
 		
 		
@@ -526,6 +544,11 @@ void FMovieGraphDeferredPass::Render(const FMovieGraphTraversalContext& InFrameT
 						RHICmdList.Transition(FRHITransitionInfo(RenderTargetResource->GetRenderTargetTexture(), ERHIAccess::RTV, ERHIAccess::SRVGraphicsPixel));
 					});
 
+				// After submission, if we're paging to system memory, mark the resources for download into system memory.
+				if (ViewState && ParentNodeThisFrame->GetEnableHistoryPerTile() && ParentNodeThisFrame->GetEnablePageToSystemMemory())
+				{
+					ViewState->SystemMemoryMirrorBackup(SystemMemoryMirror.Get());
+				}
 
 				// Data may not be something we actually want to read back and write to disk after rendering.
 				if (bDiscardOutput)
diff --git a/Engine/Plugins/MovieScene/MovieRenderPipeline/Source/MovieRenderPipelineRenderPasses/Public/Graph/Renderers/MovieGraphDeferredPanoramicPass.h b/Engine/Plugins/MovieScene/MovieRenderPipeline/Source/MovieRenderPipelineRenderPasses/Public/Graph/Renderers/MovieGraphDeferredPanoramicPass.h
--- a/Engine/Plugins/MovieScene/MovieRenderPipeline/Source/MovieRenderPipelineRenderPasses/Public/Graph/Renderers/MovieGraphDeferredPanoramicPass.h
+++ b/Engine/Plugins/MovieScene/MovieRenderPipeline/Source/MovieRenderPipelineRenderPasses/Public/Graph/Renderers/MovieGraphDeferredPanoramicPass.h
@@ -5,6 +5,8 @@
 #include "Graph/Renderers/MovieGraphImagePassBase.h"
 #include "Graph/Nodes/MovieGraphDeferredPanoramicPassNode.h"
 
+struct FSceneViewStateSystemMemoryMirror;
+
 namespace UE::MovieGraph::Rendering
 {
 	struct MOVIERENDERPIPELINERENDERPASSES_API FMovieGraphDeferredPanoramicPass : public FMovieGraphImagePassBase
@@ -37,7 +39,9 @@ namespace UE::MovieGraph::Rendering
 		// When using an auto exposure render pass, holds view states for 6 cube faces
 		TArray<FSceneViewStateReference> AutoExposureViewStates;
 
+		// Used when using Page to System Memory
 		TPimplPtr<FSceneViewStateSystemMemoryMirror> SystemMemoryMirror;
+
 		bool bHasPrintedRenderingInfo = false;
 		bool bHasPrintedWarnings = false;
 		FMovieGraphRenderDataIdentifier RenderDataIdentifier;
diff --git a/Engine/Plugins/MovieScene/MovieRenderPipeline/Source/MovieRenderPipelineRenderPasses/Public/Graph/Renderers/MovieGraphDeferredPass.h b/Engine/Plugins/MovieScene/MovieRenderPipeline/Source/MovieRenderPipelineRenderPasses/Public/Graph/Renderers/MovieGraphDeferredPass.h
--- a/Engine/Plugins/MovieScene/MovieRenderPipeline/Source/MovieRenderPipelineRenderPasses/Public/Graph/Renderers/MovieGraphDeferredPass.h
+++ b/Engine/Plugins/MovieScene/MovieRenderPipeline/Source/MovieRenderPipelineRenderPasses/Public/Graph/Renderers/MovieGraphDeferredPass.h
@@ -3,6 +3,8 @@
 
 #include "Graph/Renderers/MovieGraphImagePassBase.h"
 
+struct FSceneViewStateSystemMemoryMirror;
+
 namespace UE::MovieGraph::Rendering
 {
 	struct FMovieGraphPostRendererSubmissionParams
@@ -43,6 +45,9 @@ namespace UE::MovieGraph::Rendering
 		// Scene View history used by the renderer. When using an auto-exposure pass it'll use (-1, -1), otherwise one-per tile (and one at 0,0 if not using tiling).
 		TMap<FIntPoint, FSceneViewStateReference> SceneViewStates;
 
+		// Used when using Page to System Memory
+		TPimplPtr<FSceneViewStateSystemMemoryMirror> SystemMemoryMirror;
+
 		// The number of frames to delay to send frames from SubmissionQueue to post-render submission.
 		int32 FramesToDelayPostSubmission;
 		
