MoviePipeline: Updated icons used for Quick Render. Also slightly tweaked the wording and order of the options in the Quick Render menu to clarify their functionality.

diff --git a/Engine/Plugins/MovieScene/MovieRenderPipeline/Source/MovieRenderPipelineCore/Public/Graph/MovieGraphQuickRenderSettings.h b/Engine/Plugins/MovieScene/MovieRenderPipeline/Source/MovieRenderPipelineCore/Public/Graph/MovieGraphQuickRenderSettings.h
--- a/Engine/Plugins/MovieScene/MovieRenderPipeline/Source/MovieRenderPipelineCore/Public/Graph/MovieGraphQuickRenderSettings.h
+++ b/Engine/Plugins/MovieScene/MovieRenderPipeline/Source/MovieRenderPipelineCore/Public/Graph/MovieGraphQuickRenderSettings.h
@@ -27,17 +27,17 @@ enum class EMovieGraphQuickRenderMode : uint8
 	/** Renders the level sequence that is currently active in Sequencer. */
 	CurrentSequence,
 
-	/** Renders the shot that Sequencer's playhead is currently at. This mode is currently not implemented. */
-	/* CurrentShotAtPlayhead, */
+	/** Renders the level sequence active in Sequencer, but uses the viewport's camera instead (for the entire duration of the render). */
+	UseViewportCameraInSequence,
 
 	/** Renders a single frame using the current viewport's camera. */
-	CurrentViewport,
+	CurrentViewport UMETA(DisplayName="Current Viewport (Still)"),
 
 	/** Renders one frame for each camera actor selected in the Outliner. */
-	SelectedCameras,
+	SelectedCameras UMETA(DisplayName="Selected Cameras (Still)"),
 
-	/** Renders the level sequence active in Sequencer, but uses the viewport's camera instead (for the entire duration of the render). */
-	UseViewportCameraInSequence
+	/** Renders the shot that Sequencer's playhead is currently at. This mode is currently not implemented. */
+	/* CurrentShotAtPlayhead, */
 };
 
 /** The frame range within the active level sequence that a Quick Render should use. */
diff --git a/Engine/Plugins/MovieScene/MovieRenderPipeline/Source/MovieRenderPipelineEditor/Private/MoviePipelineQuickRenderMenu.cpp b/Engine/Plugins/MovieScene/MovieRenderPipeline/Source/MovieRenderPipelineEditor/Private/MoviePipelineQuickRenderMenu.cpp
--- a/Engine/Plugins/MovieScene/MovieRenderPipeline/Source/MovieRenderPipelineEditor/Private/MoviePipelineQuickRenderMenu.cpp
+++ b/Engine/Plugins/MovieScene/MovieRenderPipeline/Source/MovieRenderPipelineEditor/Private/MoviePipelineQuickRenderMenu.cpp
@@ -35,6 +35,20 @@ void FMoviePipelineQuickRenderMenu::AddQuickRenderButtonToToolMenu(UToolMenu* In
 	{
 		return;
 	}
+
+	// Use a lambda to fetch the button icon so it can change dynamically with the active mode
+	const TAttribute<FSlateIcon> ToolbarIconAttribute = TAttribute<FSlateIcon>::CreateLambda([]()
+	{
+		const EMovieGraphQuickRenderButtonMode ButtonMode = FMoviePipelineQuickRenderUIState::GetQuickRenderButtonMode();
+		const EMovieGraphQuickRenderMode QuickRenderMode = FMoviePipelineQuickRenderUIState::GetQuickRenderMode();
+
+		if (ButtonMode == EMovieGraphQuickRenderButtonMode::QuickRender)
+		{
+			return GetIconForQuickRenderMode(QuickRenderMode);
+		}
+		
+		return FSlateIcon(FMovieRenderPipelineStyle::Get().GetStyleSetName(), "MovieRenderPipeline.QuickRender.Icon.MovieRenderQueueMode");
+	});
 	
 	// Add to the existing "Content" section in the Assets toolbar
 	FToolMenuSection& Section = InMenu->FindOrAddSection("Content");
@@ -48,10 +62,15 @@ void FMoviePipelineQuickRenderMenu::AddQuickRenderButtonToToolMenu(UToolMenu* In
 			TAttribute<FText>::Create(
 				TAttribute<FText>::FGetter::CreateLambda([]()
 				{
-					return LOCTEXT("BeginQuickRender", "Begin Quick Render\n\nQuick Render performs a render without having to manually configure a queue and a graph. A typical quick render uses the current map and level sequence, the level sequence's playback range, and the viewport's look, to generate frames. Several different types of quick renders are available -- see the options drop-down.");
+					return LOCTEXT("BeginQuickRender", 
+									"Begin Render\n\nThere are two modes to choose from:"
+									"\n- Movie Render Queue. This uses the job(s) that are active in the Movie Render Queue editor as the source of the render."
+									"\n\n- Quick Render. This performs a render without having to manually configure a queue and a graph. A typical quick render uses "
+									"the current map and level sequence, the level sequence's playback range, and the viewport's look, to generate frames. Several different types of quick "
+									"renders are available -- see the options drop-down.");
 				})
 			),
-			FSlateIcon(FMovieRenderPipelineStyle::Get().GetStyleSetName(), "MovieRenderPipeline.QuickRender.Icon.Toolbar")
+			ToolbarIconAttribute
 		);
 	
 	Section.AddEntry(QuickRenderLaunchButton);
@@ -238,6 +257,12 @@ TSharedRef<SWidget> FMoviePipelineQuickRenderMenu::GenerateQuickRenderSettingsWi
 				{
 					return StaticEnum<EMovieGraphQuickRenderMode>()->GetDisplayNameTextByValue(static_cast<int64>(GetQuickRenderMode()));
 				})
+				.ToolTipText_Lambda([GetQuickRenderMode]()
+				{
+					const UEnum* ModeEnum = StaticEnum<EMovieGraphQuickRenderMode>();
+					
+					return ModeEnum->GetToolTipTextByIndex(ModeEnum->GetIndexByValue(static_cast<int64>(GetQuickRenderMode())));
+				})
 			];
 	};
 	
@@ -322,15 +347,15 @@ void FMoviePipelineQuickRenderMenu::GenerateModesMenuSection(UToolMenu* InMenu)
 	AddModeAction(
 		FName("QuickRender_MovieRenderQueueMode"),
 		LOCTEXT("QuickRenderMode_MovieRenderQueueLabel", "Movie Render Queue"),
-		LOCTEXT("QuickRenderMode_MovieRenderQueueToolTip", "Render with Movie Render Queue"),
-		FSlateIcon(FAppStyle::Get().GetStyleSetName(), "LevelEditor.OpenCinematic"),
+		LOCTEXT("QuickRenderMode_MovieRenderQueueToolTip", "Performs a render with the job(s) currently active in the Movie Render Queue editor."),
+		FSlateIcon(FMovieRenderPipelineStyle::Get().GetStyleSetName(), "MovieRenderPipeline.QuickRender.Icon.MovieRenderQueueMode"),
 		EMovieGraphQuickRenderButtonMode::NormalMovieRenderQueue);
 
 	AddModeAction(
 		FName("QuickRender_QuickRenderMode"),
 		LOCTEXT("QuickRenderMode_QuickRenderLabel", "Quick Render"),
-		LOCTEXT("QuickRenderMode_QuickRenderToolTip", "Render with Quick Render"),
-		FSlateIcon(FAppStyle::Get().GetStyleSetName(), "LevelEditor.OpenCinematic"),
+		LOCTEXT("QuickRenderMode_QuickRenderToolTip", "Performs a render using the Quick Render settings."),
+		FSlateIcon(FMovieRenderPipelineStyle::Get().GetStyleSetName(), "MovieRenderPipeline.QuickRender.Icon.QuickRenderMode"),
 		EMovieGraphQuickRenderButtonMode::QuickRender);
 }
 
@@ -338,8 +363,14 @@ void FMoviePipelineQuickRenderMenu::GenerateQuickRenderMenuSection(UToolMenu* In
 {
 	FToolMenuSection& Section = InMenu->AddSection("MoviePipelineQuickRenderType", LOCTEXT("QuickRenderType", "Quick Render"));
 
-	auto AddQuickRenderModeAction = [&Section](const FName& InName, const FText& InLabel, const FText& InToolTip, EMovieGraphQuickRenderMode CheckedValue)
+	auto AddQuickRenderModeAction = [&Section](const FName& InName, EMovieGraphQuickRenderMode CheckedValue)
 	{
+		const UEnum* ModeEnum = StaticEnum<EMovieGraphQuickRenderMode>();
+		const int32 ValueIndex = ModeEnum->GetIndexByValue(static_cast<int64>(CheckedValue));
+		
+		const FText ActionTooltip = ModeEnum->GetToolTipTextByIndex(ValueIndex);
+		const FText ActionLabel = ModeEnum->GetDisplayNameTextByIndex(ValueIndex);
+		
 		FToolUIAction NewAction;
 		NewAction.ExecuteAction = FToolMenuExecuteAction::CreateLambda([CheckedValue](const FToolMenuContext&)
 		{
@@ -353,40 +384,29 @@ void FMoviePipelineQuickRenderMenu::GenerateQuickRenderMenuSection(UToolMenu* In
 				? ECheckBoxState::Checked
 				: ECheckBoxState::Unchecked;
 		});
-		Section.AddMenuEntry(InName, InLabel, InToolTip, GetIconForQuickRenderMode(CheckedValue), NewAction, EUserInterfaceActionType::RadioButton);
+		Section.AddMenuEntry(InName, ActionLabel, ActionTooltip, GetIconForQuickRenderMode(CheckedValue), NewAction, EUserInterfaceActionType::RadioButton);
 	};
 
 	AddQuickRenderModeAction(
 		FName("QuickRender_CurrentSequence"),
-		LOCTEXT("QuickRender_CurrentSequenceLabel", "Current Sequence"),
-		LOCTEXT("QuickRender_CurrentSequenceToolTip", "Render the sequence that is currently active in Sequencer."),
 		EMovieGraphQuickRenderMode::CurrentSequence);
 
+	AddQuickRenderModeAction(
+		FName("QuickRender_UseViewportCamera"),
+		EMovieGraphQuickRenderMode::UseViewportCameraInSequence);
+
 	// The "Current Shot at Playhead" mode is disabled for now, may be added in the future.
 	/*AddQuickRenderModeAction(
 		FName("QuickRender_CurrentShotAtPlayhead"),
-		LOCTEXT("QuickRender_CurrentShotAtPlayheadLabel", "Current Shot at Playhead"),
-		LOCTEXT("QuickRender_CurrentShotAtPlayheadToolTip", "Render the current shot starting at the playhead."),
-		FSlateIcon("LevelSequenceEditorStyle", "LevelSequenceEditor.CreateNewLevelSequenceWithShotsInLevel"),
 		EMovieGraphQuickRenderMode::CurrentShotAtPlayhead);*/
 
 	AddQuickRenderModeAction(
 		FName("QuickRender_CurrentViewport"),
-		LOCTEXT("QuickRender_CurrentViewportLabel", "Current Viewport"),
-		LOCTEXT("QuickRender_CurrentViewportToolTip", "Render the current viewport."),
 		EMovieGraphQuickRenderMode::CurrentViewport);
 
 	AddQuickRenderModeAction(
 		FName("QuickRender_SelectedCameras"),
-		LOCTEXT("QuickRender_SelectedCamerasLabel", "Selected Camera(s)"),
-		LOCTEXT("QuickRender_SelectedCamerasToolTip", "Render the selected cameras."),
 		EMovieGraphQuickRenderMode::SelectedCameras);
-
-	AddQuickRenderModeAction(
-		FName("QuickRender_UseViewportCamera"),
-		LOCTEXT("QuickRender_UseViewportCameraLabel", "Use Viewport Camera in Sequence"),
-		LOCTEXT("QuickRender_UseViewportCameraToolTip", "Use the current viewport camera in the active level sequence."),
-		EMovieGraphQuickRenderMode::UseViewportCameraInSequence);
 }
 
 void FMoviePipelineQuickRenderMenu::GenerateQuickRenderConfigurationMenuSection(UToolMenu* InMenu)
@@ -481,7 +501,7 @@ void FMoviePipelineQuickRenderMenu::GenerateOutputMenuSection(UToolMenu* InMenu)
 		"QuickRender_PlayLastRender",
 		LOCTEXT("QuickRender_PlayLastRenderLabel", "Play Last Render"),
 		LOCTEXT("QuickRender_PlayLastRenderToolTip", "Play the media from the last time that Quick Render ran."),
-		FSlateIcon(FAppStyle::GetAppStyleSetName(), "GenericPlay"),
+		FSlateIcon(FMovieRenderPipelineStyle::Get().GetStyleSetName(), "MovieRenderPipeline.QuickRender.Icon.PlayLastRender"),
 		PlayLastRenderAction);
 
 	// Menu entry for "Open Output Directory"
@@ -544,13 +564,13 @@ FSlateIcon FMoviePipelineQuickRenderMenu::GetIconForQuickRenderMode(const EMovie
 	switch(QuickRenderMode)
 	{
 	case EMovieGraphQuickRenderMode::CurrentSequence:
-		return FSlateIcon("LevelSequenceEditorStyle", "LevelSequenceEditor.CreateNewLevelSequenceInLevel");
+		return FSlateIcon(FMovieRenderPipelineStyle::Get().GetStyleSetName(), "MovieRenderPipeline.QuickRender.Icon.CurrentSequenceMode");
 	case EMovieGraphQuickRenderMode::CurrentViewport:
-		return FSlateIcon(FAppStyle::Get().GetStyleSetName(), "LevelEditor.Tabs.Viewports");
+		return FSlateIcon(FMovieRenderPipelineStyle::Get().GetStyleSetName(), "MovieRenderPipeline.QuickRender.Icon.CurrentViewportMode");
 	case EMovieGraphQuickRenderMode::SelectedCameras:
-		return FSlateIcon(FAppStyle::Get().GetStyleSetName(), "ClassIcon.CameraActor");
+		return FSlateIcon(FMovieRenderPipelineStyle::Get().GetStyleSetName(), "MovieRenderPipeline.QuickRender.Icon.SelectedCamerasMode");
 	case EMovieGraphQuickRenderMode::UseViewportCameraInSequence:
-		return FSlateIcon(FAppStyle::Get().GetStyleSetName(), "ClassIcon.CameraActor");
+		return FSlateIcon(FMovieRenderPipelineStyle::Get().GetStyleSetName(), "MovieRenderPipeline.QuickRender.Icon.ViewportCameraInSequenceMode");
 	default:
 		return FSlateIcon();
 	}
diff --git a/Engine/Plugins/MovieScene/MovieRenderPipeline/Source/MovieRenderPipelineEditor/Private/MovieRenderPipelineStyle.cpp b/Engine/Plugins/MovieScene/MovieRenderPipeline/Source/MovieRenderPipelineEditor/Private/MovieRenderPipelineStyle.cpp
--- a/Engine/Plugins/MovieScene/MovieRenderPipeline/Source/MovieRenderPipelineEditor/Private/MovieRenderPipelineStyle.cpp
+++ b/Engine/Plugins/MovieScene/MovieRenderPipeline/Source/MovieRenderPipelineEditor/Private/MovieRenderPipelineStyle.cpp
@@ -107,7 +107,13 @@ FMovieRenderPipelineStyle::FMovieRenderPipelineStyle()
 		.SetDownArrowPadding(FMargin(0)));
 
 	// Quick Render icons
-	Set("MovieRenderPipeline.QuickRender.Icon.Toolbar", new IMAGE_BRUSH_SVG("QuickRender/QuickRenderToolbar_16", Icon16x16));
+	Set("MovieRenderPipeline.QuickRender.Icon.MovieRenderQueueMode", new IMAGE_BRUSH_SVG("QuickRender/MovieRenderQueueMode_16", Icon16x16));
+	Set("MovieRenderPipeline.QuickRender.Icon.QuickRenderMode", new IMAGE_BRUSH_SVG("QuickRender/QuickRenderMode_16", Icon16x16));
+	Set("MovieRenderPipeline.QuickRender.Icon.CurrentSequenceMode", new IMAGE_BRUSH_SVG("QuickRender/CurrentSequenceMode_16", Icon16x16));
+	Set("MovieRenderPipeline.QuickRender.Icon.ViewportCameraInSequenceMode", new IMAGE_BRUSH_SVG("QuickRender/ViewportCameraInSequenceMode_16", Icon16x16));
+	Set("MovieRenderPipeline.QuickRender.Icon.CurrentViewportMode", new IMAGE_BRUSH_SVG("QuickRender/CurrentViewportMode_16", Icon16x16));
+	Set("MovieRenderPipeline.QuickRender.Icon.SelectedCamerasMode", new IMAGE_BRUSH_SVG("QuickRender/SelectedCamerasMode_16", Icon16x16));
+	Set("MovieRenderPipeline.QuickRender.Icon.PlayLastRender", new IMAGE_BRUSH_SVG("QuickRender/PlayLastRender_16", Icon16x16));
 
 	FSlateStyleRegistry::RegisterSlateStyle(*this);
 }
