MoviePipeline: Fixed an issue where variable assignments were not properly refreshed if undoing a change of the graph config within Quick Render settings.

diff --git a/Engine/Plugins/MovieScene/MovieRenderPipeline/Source/MovieRenderPipelineCore/Private/Graph/MovieGraphQuickRenderSettings.cpp b/Engine/Plugins/MovieScene/MovieRenderPipeline/Source/MovieRenderPipelineCore/Private/Graph/MovieGraphQuickRenderSettings.cpp
--- a/Engine/Plugins/MovieScene/MovieRenderPipeline/Source/MovieRenderPipelineCore/Private/Graph/MovieGraphQuickRenderSettings.cpp
+++ b/Engine/Plugins/MovieScene/MovieRenderPipeline/Source/MovieRenderPipelineCore/Private/Graph/MovieGraphQuickRenderSettings.cpp
@@ -3,6 +3,7 @@
 #include "Graph/MovieGraphQuickRenderSettings.h"
 
 #include "Misc/CoreDelegates.h"
+#include "Misc/TransactionObjectEvent.h"
 #include "MoviePipelineUtils.h"
 #include "MovieRenderPipelineCoreModule.h"
 #include "UObject/Package.h"
@@ -127,14 +128,6 @@ UMovieJobVariableAssignmentContainer* UMovieGraphQuickRenderModeSettings::GetVar
 #if WITH_EDITOR
 void UMovieGraphQuickRenderModeSettings::PostEditChangeProperty(FPropertyChangedEvent& PropertyChangedEvent)
 {
-	if (PropertyChangedEvent.GetPropertyName() == GET_MEMBER_NAME_CHECKED(UMovieGraphQuickRenderModeSettings, GraphPreset))
-	{
-		// Refresh the variable assignments if the graph preset changes
-		RefreshVariableAssignments(this);
-		
-		OnGraphChangedDelegate.Broadcast();
-	}
-
 	// Let the parent settings object know that it is dirty and should queue up a save
 	if (UMovieGraphQuickRenderSettings* QuickRenderSettings = GetTypedOuter<UMovieGraphQuickRenderSettings>())
 	{
@@ -143,4 +136,18 @@ void UMovieGraphQuickRenderModeSettings::PostEditChangeProperty(FPropertyChanged
 	
 	Super::PostEditChangeProperty(PropertyChangedEvent);
 }
+
+void UMovieGraphQuickRenderModeSettings::PostTransacted(const FTransactionObjectEvent& TransactionEvent)
+{
+	Super::PostTransacted(TransactionEvent);
+
+	// Refresh the variable assignments if the graph preset changes. This is done in PostTransacted() rather than PostEditChangeProperty() so the
+	// change to GraphPreset can be picked up when the user changes it directly AND when changed through undo/redo.
+	const TArray<FName>& ChangedProperties = TransactionEvent.GetChangedProperties();
+	if (ChangedProperties.Contains(GET_MEMBER_NAME_CHECKED(UMovieGraphQuickRenderModeSettings, GraphPreset)))
+	{
+		RefreshVariableAssignments(this);
+		OnGraphChangedDelegate.Broadcast();
+	}
+}
 #endif	// WITH_EDITOR
\ No newline at end of file
diff --git a/Engine/Plugins/MovieScene/MovieRenderPipeline/Source/MovieRenderPipelineCore/Public/Graph/MovieGraphQuickRenderSettings.h b/Engine/Plugins/MovieScene/MovieRenderPipeline/Source/MovieRenderPipelineCore/Public/Graph/MovieGraphQuickRenderSettings.h
--- a/Engine/Plugins/MovieScene/MovieRenderPipeline/Source/MovieRenderPipelineCore/Public/Graph/MovieGraphQuickRenderSettings.h
+++ b/Engine/Plugins/MovieScene/MovieRenderPipeline/Source/MovieRenderPipelineCore/Public/Graph/MovieGraphQuickRenderSettings.h
@@ -197,6 +197,7 @@ private:
 #if WITH_EDITOR
 	//~ Begin UObject Interface
 	virtual void PostEditChangeProperty(FPropertyChangedEvent& PropertyChangedEvent) override;
+	virtual void PostTransacted(const FTransactionObjectEvent& TransactionEvent) override;
 	//~ End UObject Interface
 #endif
 };
\ No newline at end of file
