MoviePipeline: Fixed a bug that caused the Select and Branch nodes to not properly resolve a value on the pin that controls their behavior if the pin was connected to a subgraph or controlled via the queue.

Also fixed an issue that caused variables to not be overridable in the queue in some scenarios, and a bug that caused Quick Render to not respect the override (checked) state of variables in its settings dialog.

diff --git a/Engine/Plugins/MovieScene/MovieRenderPipeline/Source/MovieRenderPipelineCore/Private/Graph/MovieGraphConfig.cpp b/Engine/Plugins/MovieScene/MovieRenderPipeline/Source/MovieRenderPipelineCore/Private/Graph/MovieGraphConfig.cpp
--- a/Engine/Plugins/MovieScene/MovieRenderPipeline/Source/MovieRenderPipelineCore/Private/Graph/MovieGraphConfig.cpp
+++ b/Engine/Plugins/MovieScene/MovieRenderPipeline/Source/MovieRenderPipelineCore/Private/Graph/MovieGraphConfig.cpp
@@ -1709,53 +1709,14 @@ void UMovieGraphConfig::CopyOverriddenProperties(UMovieGraphNode* FromNode, UMov
 				TArray<UMovieGraphPin*> ConnectionPath;
 				
 				// Iterate up the connection chain and find all pins which might have a value that can be resolved.
-				FMovieGraphEvaluationContext ValueConnectionContext;
+				FMovieGraphEvaluationContext ValueConnectionContext = InEvaluationContext;
 				ValueConnectionContext.PinBeingFollowed = InputPin;
-				ValueConnectionContext.SubgraphStack = InEvaluationContext.SubgraphStack;
 				TArray<UMovieGraphPin*> ConnectedValuePins = InputPin->Node->EvaluatePinsToFollow(ValueConnectionContext);
-				while (!ConnectedValuePins.IsEmpty())
-				{
-					UMovieGraphPin* ConnectedValuePin = ConnectedValuePins[0];
-					if (!ensureMsgf(ConnectedValuePin, TEXT("Found an invalid pin on node '%s'."), *InputPin->Node->GetName()))
-					{
-						// Can't continue following the connection chain if an invalid pin was found
-						break;
-					}
-					
-					if (ConnectionPath.Contains(ConnectedValuePin))
-					{
-						// Recursive connection found
-						UE_LOG(LogMovieRenderPipeline, Warning, TEXT("Found a cycle when following the data connection on pin '%s' for node '%s'. Value will not be resolved."),
-							*ConnectedValuePin->Properties.Label.ToString(), *FromNode->GetName());
-						break;
-					}
-
-					// For the connected value to be used, the types must be compatible and the node the value is originating from must be enabled
-					if (InputPin->IsTypeCompatibleWith(ConnectedValuePin) && ConnectedValuePin->Node && !ConnectedValuePin->Node->IsDisabled())
-					{
-						ConnectionPath.Add(ConnectedValuePin);
-					}
 
-					ValueConnectionContext.PinBeingFollowed = ConnectedValuePin;
-					ConnectedValuePins = ConnectedValuePin->Node->EvaluatePinsToFollow(ValueConnectionContext);
-				}
-
-				// Work backwards and use the most upstream value that can be resolved. The most upstream values wins. For example, if a node has an
-				// exposed pin, that pin is connected to a subgraph's input, and that input is then connected to a variable node in the parent graph.
-				// The variable node's value should be used if it can be resolved, not the subgraph's input value.
-				bool bFoundResolvedValue = false;
-				for (int32 Index = ConnectionPath.Num() - 1; Index >= 0; --Index)
+				FString ResolvedValue;
+				const bool bFoundResolvedValue = UE::MovieGraph::ResolveConnectedPinValue(InputPin, ConnectedValuePins, ValueConnectionContext, ResolvedValue);
+				if (bFoundResolvedValue)
 				{
-					const UMovieGraphPin* ConnectedPin = ConnectionPath[Index];
-					
-					const FString ResolvedValue = ConnectedPin->Node->GetResolvedValueForOutputPin(ConnectedPin->Properties.Label, &InEvaluationContext.UserContext);
-					if (ResolvedValue.IsEmpty())
-					{
-						continue;
-					}
-
-					bFoundResolvedValue = true;
-					
 					if (bIsDynamic)
 					{
 						ToNode->SetDynamicPropertyValue(PropertyName, ResolvedValue);
@@ -1767,13 +1728,7 @@ void UMovieGraphConfig::CopyOverriddenProperties(UMovieGraphNode* FromNode, UMov
 						EditConditionProperty->SetPropertyValue_InContainer(ToNode, true);
 					}
 
-					// Resolved a value for this pin; stop iterating over the connection chain
-					break;
-				}
-				
-				// The property value was set via a connected pin; move on to the next property
-				if (bFoundResolvedValue)
-				{
+					// The property value was set via a connected pin; move on to the next property
 					continue;
 				}
 			}
diff --git a/Engine/Plugins/MovieScene/MovieRenderPipeline/Source/MovieRenderPipelineCore/Private/Graph/MovieGraphUtils.cpp b/Engine/Plugins/MovieScene/MovieRenderPipeline/Source/MovieRenderPipelineCore/Private/Graph/MovieGraphUtils.cpp
--- a/Engine/Plugins/MovieScene/MovieRenderPipeline/Source/MovieRenderPipelineCore/Private/Graph/MovieGraphUtils.cpp
+++ b/Engine/Plugins/MovieScene/MovieRenderPipeline/Source/MovieRenderPipelineCore/Private/Graph/MovieGraphUtils.cpp
@@ -2,6 +2,9 @@
 
 #include "MovieGraphUtils.h"
 
+#include "Graph/MovieGraphConfig.h"
+#include "MovieRenderPipelineCoreModule.h"
+
 #include "AudioMixerDevice.h"
 #include "AudioThread.h"
 #include "Engine/Engine.h"
@@ -34,6 +37,63 @@ namespace UE::MovieGraph
 		return NewName;
 	}
 
+	bool ResolveConnectedPinValue(UMovieGraphPin* InFromPin, const TArray<UMovieGraphPin*>& InConnectedPins, const FMovieGraphEvaluationContext& InEvaluationContext, FString& OutResolvedValue)
+	{
+		TArray<UMovieGraphPin*> ConnectionPath;
+			
+		// Iterate up the connection chain and find all pins which might have a value that can be resolved.
+		FMovieGraphEvaluationContext ValueConnectionContext;
+		ValueConnectionContext.PinBeingFollowed = InFromPin;
+		ValueConnectionContext.SubgraphStack = InEvaluationContext.SubgraphStack;
+		
+		TArray<UMovieGraphPin*> ConnectedValuePins = InConnectedPins;
+		while (!ConnectedValuePins.IsEmpty())
+		{
+			UMovieGraphPin* ConnectedValuePin = ConnectedValuePins[0];
+			if (!ensureMsgf(ConnectedValuePin, TEXT("Found an invalid pin on node '%s'."), *InFromPin->Node->GetName()))
+			{
+				// Can't continue following the connection chain if an invalid pin was found
+				break;
+			}
+			
+			if (ConnectionPath.Contains(ConnectedValuePin))
+			{
+				// Recursive connection found
+				UE_LOG(LogMovieRenderPipeline, Warning, TEXT("Found a cycle when following the data connection on pin '%s' for node '%s'. Value will not be resolved."),
+					*ConnectedValuePin->Properties.Label.ToString(), *InFromPin->Node->GetName());
+				break;
+			}
+
+			// For the connected value to be used, the types must be compatible and the node the value is originating from must be enabled
+			if (InFromPin->IsTypeCompatibleWith(ConnectedValuePin) && ConnectedValuePin->Node && !ConnectedValuePin->Node->IsDisabled())
+			{
+				ConnectionPath.Add(ConnectedValuePin);
+			}
+
+			ValueConnectionContext.PinBeingFollowed = ConnectedValuePin;
+			ConnectedValuePins = ConnectedValuePin->Node->EvaluatePinsToFollow(ValueConnectionContext);
+		}
+
+		// Work backwards and use the most upstream value that can be resolved. The most upstream value wins. For example, if a node has an
+		// exposed pin, that pin is connected to a subgraph's input, and that input is then connected to a variable node in the parent graph,
+		// then the variable node's value should be used if it can be resolved, not the subgraph's input value.
+		for (int32 Index = ConnectionPath.Num() - 1; Index >= 0; --Index)
+		{
+			const UMovieGraphPin* ConnectedPin = ConnectionPath[Index];
+				
+			const FString ResolvedValue = ConnectedPin->Node->GetResolvedValueForOutputPin(ConnectedPin->Properties.Label, &InEvaluationContext.UserContext);
+			if (ResolvedValue.IsEmpty())
+			{
+				continue;
+			}
+
+			OutResolvedValue = ResolvedValue;
+			return true;
+		}
+			
+		return false;
+	}
+
 	namespace Audio
 	{
 		FAudioDevice* GetAudioDeviceFromWorldContext(const UObject* InWorldContextObject)
diff --git a/Engine/Plugins/MovieScene/MovieRenderPipeline/Source/MovieRenderPipelineCore/Private/Graph/MovieGraphUtils.h b/Engine/Plugins/MovieScene/MovieRenderPipeline/Source/MovieRenderPipelineCore/Private/Graph/MovieGraphUtils.h
--- a/Engine/Plugins/MovieScene/MovieRenderPipeline/Source/MovieRenderPipelineCore/Private/Graph/MovieGraphUtils.h
+++ b/Engine/Plugins/MovieScene/MovieRenderPipeline/Source/MovieRenderPipelineCore/Private/Graph/MovieGraphUtils.h
@@ -6,6 +6,8 @@
 
 // Forward declares
 class FAudioDevice;
+class UMovieGraphPin;
+struct FMovieGraphEvaluationContext;
 namespace Audio
 {
 	class FMixerDevice;
@@ -20,6 +22,17 @@ namespace UE::MovieGraph
 	 */
 	FString GetUniqueName(const TArray<FString>& InExistingNames, const FString& InBaseName);
 
+	/**
+	 * Gets the resolved value of an input pin (InFromPin), given all of its connections (InConnectedPins). Providing an evaluation context is needed,
+	 * mostly to ensure that recursion doesn't occur. The resolved value is provided in OutResolvedValue as the serialized representation.
+	 * Returns true if a value could be resolved, else false.
+	 *
+	 * Generally the normal graph evaluation process should resolve values, and using this method might indicate you're doing the wrong thing.
+	 * However, for some non-setting nodes, it may be necessary to resolve values manually since non-setting nodes are not evaluated like normal
+	 * setting nodes.
+	 */
+	bool ResolveConnectedPinValue(UMovieGraphPin* InFromPin, const TArray<UMovieGraphPin*>& InConnectedPins, const FMovieGraphEvaluationContext& InEvaluationContext, FString& OutResolvedValue);
+
 	namespace Audio
 	{
 		/** Gets the audio device from the supplied world context (or nullptr if it could not be determined). */
diff --git a/Engine/Plugins/MovieScene/MovieRenderPipeline/Source/MovieRenderPipelineCore/Private/Graph/Nodes/MovieGraphBranchNode.cpp b/Engine/Plugins/MovieScene/MovieRenderPipeline/Source/MovieRenderPipelineCore/Private/Graph/Nodes/MovieGraphBranchNode.cpp
--- a/Engine/Plugins/MovieScene/MovieRenderPipeline/Source/MovieRenderPipelineCore/Private/Graph/Nodes/MovieGraphBranchNode.cpp
+++ b/Engine/Plugins/MovieScene/MovieRenderPipeline/Source/MovieRenderPipelineCore/Private/Graph/Nodes/MovieGraphBranchNode.cpp
@@ -4,6 +4,7 @@
 
 #include "Graph/MovieGraphConfig.h"
 #include "Graph/MovieGraphEdge.h"
+#include "Graph/MovieGraphUtils.h"
 #include "MovieRenderPipelineCoreModule.h"
 
 #include "Styling/AppStyle.h"
@@ -64,12 +65,15 @@ TArray<UMovieGraphPin*> UMovieGraphBranchNode::EvaluatePinsToFollow(FMovieGraphE
 		return PinsToFollow;
 	}
 
+	FString ResolvedValue;
+	UE::MovieGraph::ResolveConnectedPinValue(ConditionalPin, {OtherPin}, InContext, ResolvedValue);
+
 	// TODO: Comparing against a string like this is not ideal. This could use GetPropertyValueContainerForPin() instead,
 	// however that may not work in a general case. This fits variables nicely since they already have a
 	// UMovieGraphValueContainer -- other nodes in the future may not, though, and generating them on every frame may be expensive.
 	// Still an area that needs to be researched.
 	static const FString TrueValue = FString("true");
-	if (ConnectedNode->GetResolvedValueForOutputPin(OtherPin->Properties.Label, &InContext.UserContext) == TrueValue)
+	if (ResolvedValue == TrueValue)
 	{
 		PinsToFollow.Add(GetInputPin(UE::MovieGraph::BranchNode::TrueBranch));
 	}
diff --git a/Engine/Plugins/MovieScene/MovieRenderPipeline/Source/MovieRenderPipelineCore/Private/Graph/Nodes/MovieGraphSelectNode.cpp b/Engine/Plugins/MovieScene/MovieRenderPipeline/Source/MovieRenderPipelineCore/Private/Graph/Nodes/MovieGraphSelectNode.cpp
--- a/Engine/Plugins/MovieScene/MovieRenderPipeline/Source/MovieRenderPipelineCore/Private/Graph/Nodes/MovieGraphSelectNode.cpp
+++ b/Engine/Plugins/MovieScene/MovieRenderPipeline/Source/MovieRenderPipelineCore/Private/Graph/Nodes/MovieGraphSelectNode.cpp
@@ -4,6 +4,7 @@
 
 #include "MovieRenderPipelineCoreModule.h"
 #include "Graph/MovieGraphConfig.h"
+#include "Graph/MovieGraphUtils.h"
 #include "Styling/AppStyle.h"
 
 namespace UE::MovieGraph::SelectNode
@@ -184,14 +185,20 @@ TArray<UMovieGraphPin*> UMovieGraphSelectNode::EvaluatePinsToFollow(FMovieGraphE
 
 	// Try getting the value from a connection first
 	bool bGotValueFromConnection = false;
-	if (const UMovieGraphPin* SelectPin = GetInputPin(UE::MovieGraph::SelectNode::SelectedOption, EMovieGraphPinQueryRequirement::BuiltIn))
+	if (UMovieGraphPin* SelectPin = GetInputPin(UE::MovieGraph::SelectNode::SelectedOption, EMovieGraphPinQueryRequirement::BuiltIn))
 	{
 		constexpr bool bFollowRerouteConnections = true;
-		if (const UMovieGraphPin* OtherPin = SelectPin->GetFirstConnectedPin(bFollowRerouteConnections))
+		if (UMovieGraphPin* OtherPin = SelectPin->GetFirstConnectedPin(bFollowRerouteConnections))
 		{
-			if (const UMovieGraphNode* ConnectedNode = OtherPin->Node)
+			FString ResolvedValue;
+			bGotValueFromConnection = UE::MovieGraph::ResolveConnectedPinValue(SelectPin, {OtherPin}, InContext, ResolvedValue);
+
+			if (bGotValueFromConnection)
 			{
-				bGotValueFromConnection = ConnectedNode->GetResolvedValueForOutputPin(OtherPin->Properties.Label, &InContext.UserContext, ResolvedSelectedOption);
+				ResolvedSelectedOption = NewObject<UMovieGraphValueContainer>();
+				ResolvedSelectedOption->SetValueType(OtherPin->Properties.Type);
+				ResolvedSelectedOption->SetValueTypeObject(OtherPin->Properties.TypeObject);
+				ResolvedSelectedOption->SetValueSerializedString(ResolvedValue);
 			}
 		}
 	}
diff --git a/Engine/Plugins/MovieScene/MovieRenderPipeline/Source/MovieRenderPipelineEditor/Private/Graph/MovieGraphQuickRender.cpp b/Engine/Plugins/MovieScene/MovieRenderPipeline/Source/MovieRenderPipelineEditor/Private/Graph/MovieGraphQuickRender.cpp
--- a/Engine/Plugins/MovieScene/MovieRenderPipeline/Source/MovieRenderPipelineEditor/Private/Graph/MovieGraphQuickRender.cpp
+++ b/Engine/Plugins/MovieScene/MovieRenderPipeline/Source/MovieRenderPipelineEditor/Private/Graph/MovieGraphQuickRender.cpp
@@ -298,9 +298,12 @@ bool UMovieGraphQuickRenderSubsystem::GenerateAndPopulateQueue(TObjectPtr<UMovie
 
 		for (const UMovieGraphVariable* GraphVariable : (*DuplicateGraph)->GetVariables())
 		{
+			bool bIsVariableAssignmentEnabled = false;
+			IncomingVariableAssignment->GetVariableAssignmentEnableState(GraphVariable, bIsVariableAssignmentEnabled);
+			
 			// Copying values by serialized string isn't ideal, but it's the easiest way to transfer values between the variable assignments
 			NewVariableAssignments->SetValueSerializedString(GraphVariable, IncomingVariableAssignment->GetValueSerializedString(GraphVariable));
-			NewVariableAssignments->SetVariableAssignmentEnableState(GraphVariable, true);
+			NewVariableAssignments->SetVariableAssignmentEnableState(GraphVariable, bIsVariableAssignmentEnabled);
 		}
 	}
 
