MovieGraph: Added high-res support to the Path Tracer node, fixed an issue with motion blur only working on the last tile when using High Resolution Tiling, fixed Page to System Memory not being exposable as a pin on the node.

diff --git a/Engine/Plugins/MovieScene/MovieRenderPipeline/Source/MovieRenderPipelineRenderPasses/Private/Graph/Renderers/MovieGraphDeferredPass.cpp b/Engine/Plugins/MovieScene/MovieRenderPipeline/Source/MovieRenderPipelineRenderPasses/Private/Graph/Renderers/MovieGraphDeferredPass.cpp
--- a/Engine/Plugins/MovieScene/MovieRenderPipeline/Source/MovieRenderPipelineRenderPasses/Private/Graph/Renderers/MovieGraphDeferredPass.cpp
+++ b/Engine/Plugins/MovieScene/MovieRenderPipeline/Source/MovieRenderPipelineRenderPasses/Private/Graph/Renderers/MovieGraphDeferredPass.cpp
@@ -335,7 +335,7 @@ void FMovieGraphDeferredPass::Render(const FMovieGraphTraversalContext& InFrameT
 				// scene multiple times.
 				const bool bIsLastTile = FIntPoint(TileX, TileY) == FIntPoint(TileCount.X - 1, TileCount.Y - 1);
 				const bool bHasTiles = (TileCount.X * TileCount.Y) > 1;
-				const bool bWorldIsPaused = !(bIsLastTile && SpatialIndex == (NumSpatialSamples - 1));
+				const bool bWorldIsPaused = !(SpatialIndex == (NumSpatialSamples - 1));
 				const int32 FrameIndex = InTimeData.RenderedFrameNumber * ((InTimeData.TemporalSampleCount * NumSpatialSamples) + (InTimeData.TemporalSampleIndex * NumSpatialSamples)) + SpatialIndex;
 
 				// We only allow a spatial jitter if we have more than one sample
diff --git a/Engine/Plugins/MovieScene/MovieRenderPipeline/Source/MovieRenderPipelineRenderPasses/Public/Graph/Nodes/MovieGraphDeferredPassNode.h b/Engine/Plugins/MovieScene/MovieRenderPipeline/Source/MovieRenderPipelineRenderPasses/Public/Graph/Nodes/MovieGraphDeferredPassNode.h
--- a/Engine/Plugins/MovieScene/MovieRenderPipeline/Source/MovieRenderPipelineRenderPasses/Public/Graph/Nodes/MovieGraphDeferredPassNode.h
+++ b/Engine/Plugins/MovieScene/MovieRenderPipeline/Source/MovieRenderPipelineRenderPasses/Public/Graph/Nodes/MovieGraphDeferredPassNode.h
@@ -189,6 +189,6 @@ public:
 	* but can allow using Lumen and other features that rely on bAllocateHistoryPerTile to be used on GPUs that do not have enough VRAM to store all of the history
 	* data for every tile at once.
 	*/
-	UPROPERTY(BlueprintReadWrite, Category = "High Resolution Tiling", meta = (EditCondition = "bOverride_bPageToSystemMemory"))
+	UPROPERTY(EditAnywhere, BlueprintReadWrite, Category = "High Resolution Tiling", meta = (EditCondition = "bOverride_bPageToSystemMemory"))
 	bool bPageToSystemMemory;
 };
\ No newline at end of file
diff --git a/Engine/Plugins/MovieScene/MovieRenderPipeline/Source/MovieRenderPipelineRenderPasses/Public/Graph/Nodes/MovieGraphPathTracerPassNode.h b/Engine/Plugins/MovieScene/MovieRenderPipeline/Source/MovieRenderPipelineRenderPasses/Public/Graph/Nodes/MovieGraphPathTracerPassNode.h
--- a/Engine/Plugins/MovieScene/MovieRenderPipeline/Source/MovieRenderPipelineRenderPasses/Public/Graph/Nodes/MovieGraphPathTracerPassNode.h
+++ b/Engine/Plugins/MovieScene/MovieRenderPipeline/Source/MovieRenderPipelineRenderPasses/Public/Graph/Nodes/MovieGraphPathTracerPassNode.h
@@ -56,6 +56,12 @@ public:
 	virtual bool GetAllowDenoiser() const override;
 	virtual TUniquePtr<UE::MovieGraph::Rendering::FMovieGraphImagePassBase> CreateInstance() const;
 	virtual FEngineShowFlags GetShowFlags() const override;
+	virtual bool GetEnableHighResolutionTiling() const override { return bEnableHighResolutionTiling; }
+	virtual FIntPoint GetTileCount() const override { return FIntPoint(TileCount, TileCount); }
+	virtual float GetTileOverlapPercentage() const override { return OverlapPercentage; }
+	virtual bool GetEnablePageToSystemMemory() const override { return bPageToSystemMemory; }
+	virtual bool GetEnableHistoryPerTile() const override { return bAllocateHistoryPerTile; }
+
 	// ~UMovieGraphImagePassBaseNode Interface
 
 	// UMovieGraphSettingNode Interface
@@ -127,6 +133,21 @@ public:
 	UPROPERTY(EditAnywhere, BlueprintReadWrite, Category = Overrides, meta = (InlineEditConditionToggle))
 	uint8 bOverride_AdditionalPostProcessMaterials : 1;
 
+	UPROPERTY(EditAnywhere, BlueprintReadWrite, Category = Overrides, meta = (InlineEditConditionToggle))
+	uint8 bOverride_bEnableHighResolutionTiling : 1;
+
+	UPROPERTY(EditAnywhere, BlueprintReadWrite, Category = Overrides, meta = (InlineEditConditionToggle))
+	uint8 bOverride_TileCount : 1;
+
+	UPROPERTY(EditAnywhere, BlueprintReadWrite, Category = Overrides, meta = (InlineEditConditionToggle))
+	uint8 bOverride_OverlapPercentage : 1;
+
+	UPROPERTY(EditAnywhere, BlueprintReadWrite, Category = Overrides, meta = (InlineEditConditionToggle))
+	uint8 bOverride_bAllocateHistoryPerTile : 1;
+
+	UPROPERTY(EditAnywhere, BlueprintReadWrite, Category = Overrides, meta = (InlineEditConditionToggle))
+	uint8 bOverride_bPageToSystemMemory : 1;
+
 	/**
 	* How many sub-pixel jitter renders should we do per temporal sample? This can be used to achieve high
 	* sample counts without Temporal Sub-Sampling (allowing high sample counts without motion blur being enabled),
@@ -228,6 +249,59 @@ public:
 	UPROPERTY(EditAnywhere, BlueprintReadWrite, Category = "Post Process Materials", meta=(EditCondition="bOverride_AdditionalPostProcessMaterials"))
 	TArray<FMoviePipelinePostProcessPass> AdditionalPostProcessMaterials;
 
+	/**
+	* If true, the render will be done using a "tiled" render, which can overcome size limitations of GPUs but comes with a significant number
+	* of limitations. The internal GBuffer used for rendering is quite memory intensive, so a very large (ie: 8-16k) render may be impractical
+	* from a memory standpoint. You can enable High Resolution Tiling to render this in multiple smaller passes, but there is overhead to each
+	* tile as well, in terms of raytracing and lumen acceleration structures.
+	*
+	* - If you have VRAM available, it's better to render with 1 tile instead of 2, and instead increase the TDR (Timeout Device Recovery) in
+	* your OS to allow frames to take longer than the default 2s limit.
+	* - If you need to use tiling, and you need to use Lumen, TAA/TSR, or other rendering features that require the previous frame's buffer,
+	* then you'll need to enable bAllocateHistoryPerTile. This can come at a large VRAM cost but may require less vram than having a larger
+	* GBuffer resolution.
+	* - If you have spare system memory (RAM), you can use the experimental bPageToSystemMemory feature to download all of the per-tile rendering history
+	* into CPU memory after each tile, and then upload it again before the next time that tile is used. This comes at a _significant_ performance and system
+	* memory cost but can allow utilizing significantly larger overall resolutions while still supporting Lumen, TAA/TSR and other features.
+	*/
+	UPROPERTY(EditAnywhere, BlueprintReadWrite, Category = "High Resolution Tiling", meta = (EditCondition = "bOverride_bEnableHighResolutionTiling"))
+	bool bEnableHighResolutionTiling;
+
+	/**
+	* If bEnableHighResolutionTiling is enabled, what is the tile count that the screen should be broken into. This is not in pixels, but in number of tiles per
+	* side, ie: an output resolution of 4k, and a tile count of 2, produces 4 tiles (2 horizontal, 2 vertical) with each tile being 1080p. Larger tile counts
+	* shrink the individual render resolution, but increase the total number of renders needed.
+	*/
+	UPROPERTY(EditAnywhere, BlueprintReadWrite, Category = "High Resolution Tiling", meta = (EditCondition = "bOverride_TileCount", UIMin = "1", ClampMin = "1", UIMax = "16"))
+	int32 TileCount;
+
+	/**
+	* Rendering effects such as Depth of Field may produce different results near the edge of a tile (as it cannot sample outside of the tile itself), so this
+	* setting allows you to create an overlapped region between tiles. 10% is a good starting point, but may need to be increased if you have extremely large
+	* depth of field.
+	*
+	* Note: This uses 0-50 and not 0-.5 like the previous system did to bring it in-line with other usages of overscan in the engine (nDisplay).
+	*/
+	UPROPERTY(EditAnywhere, BlueprintReadWrite, Category = "High Resolution Tiling", meta = (EditCondition = "bOverride_OverlapPercentage", UIMin = "0", UIMax = "50", ClampMin = "0", ClampMax = "50"))
+	float OverlapPercentage;
+
+	/**
+	* If enabled, a FSceneViewStateInterface is allocated for each tile in the high resolution image. This is required for TAA/TSR/Lumen and other modern rendering features
+	* to work correctly, but can consume significant amounts of VRAM to store the state for each tile. This can be mitigated (at significant render time impact) with
+	* the new experimental bPageToSystemMemory cost.
+	*/
+	UPROPERTY(EditAnywhere, BlueprintReadWrite, Category = "High Resolution Tiling", meta = (EditCondition = "bOverride_bAllocateHistoryPerTile"))
+	bool bAllocateHistoryPerTile;
+
+	/**
+	* Experimental Feature: When enabled, after each tile is rendered, MRQ will download the per-tile image history back to system RAM, and then the next time
+	* the tile is rendered on the subsequent frame the data is transfered from sytem memory back to GPU memory for use. This significantly impacts rendering time,
+	* but can allow using Lumen and other features that rely on bAllocateHistoryPerTile to be used on GPUs that do not have enough VRAM to store all of the history
+	* data for every tile at once.
+	*/
+	UPROPERTY(EditAnywhere, BlueprintReadWrite, Category = "High Resolution Tiling", meta = (EditCondition = "bOverride_bPageToSystemMemory"))
+	bool bPageToSystemMemory;
+
 private:
 	/**
 	 * The original value of the "r.PathTracing.ProgressDisplay" cvar before the render starts. The progress display
