MoviePipeline: Fixed a crash that could happen if a graph was opened that contains a node which came from a plugin that isn't currently loaded.

diff --git a/Engine/Plugins/MovieScene/MovieRenderPipeline/Source/MovieRenderPipelineCore/Private/Graph/MovieGraphConfig.cpp b/Engine/Plugins/MovieScene/MovieRenderPipeline/Source/MovieRenderPipelineCore/Private/Graph/MovieGraphConfig.cpp
--- a/Engine/Plugins/MovieScene/MovieRenderPipeline/Source/MovieRenderPipelineCore/Private/Graph/MovieGraphConfig.cpp
+++ b/Engine/Plugins/MovieScene/MovieRenderPipeline/Source/MovieRenderPipelineCore/Private/Graph/MovieGraphConfig.cpp
@@ -783,7 +783,7 @@ bool UMovieGraphConfig::RemoveNode(UMovieGraphNode* InNode)
 	{
 		FFrame::KismetExecutionMessage(
 			*FString::Printf(
-				TEXT("%hs: Could not remove invalid InNode"), __FUNCTION__),
+				TEXT("%hs: Could not remove invalid node. Is it from a plugin that's not currently loaded?"), __FUNCTION__),
 				ELogVerbosity::Error);
 		return false;
 	}
diff --git a/Engine/Plugins/MovieScene/MovieRenderPipeline/Source/MovieRenderPipelineCore/Private/Graph/MovieGraphPin.cpp b/Engine/Plugins/MovieScene/MovieRenderPipeline/Source/MovieRenderPipelineCore/Private/Graph/MovieGraphPin.cpp
--- a/Engine/Plugins/MovieScene/MovieRenderPipeline/Source/MovieRenderPipelineCore/Private/Graph/MovieGraphPin.cpp
+++ b/Engine/Plugins/MovieScene/MovieRenderPipeline/Source/MovieRenderPipelineCore/Private/Graph/MovieGraphPin.cpp
@@ -175,7 +175,7 @@ bool UMovieGraphPin::IsConnected() const
 {
 	for (UMovieGraphEdge* Edge : Edges)
 	{
-		if (Edge->IsValid())
+		if (Edge && Edge->IsValid())
 		{
 			return true;
 		}
@@ -237,10 +237,13 @@ TArray<UMovieGraphPin*> UMovieGraphPin::GetAllConnectedPins() const
 
 	for (UMovieGraphEdge* Edge : Edges)
 	{
-		if (ensureMsgf(Edge, TEXT("Null edge found when trying to get connected pin!")))
+		if (!IsValid(Edge))
 		{
-			ConnectedPins.Add(Edge->GetOtherPin(this));
+			UE_LOG(LogMovieRenderPipeline, Warning, TEXT("GetAllConnectedPins() - Found an invalid edge. This is typically caused by a node that is from a plugin that is not currently loaded."));
+			continue;
 		}
+
+		ConnectedPins.Add(Edge->GetOtherPin(this));
 	}
 
 	return ConnectedPins;
diff --git a/Engine/Plugins/MovieScene/MovieRenderPipeline/Source/MovieRenderPipelineEditor/Private/Graph/MovieEdGraphNode.cpp b/Engine/Plugins/MovieScene/MovieRenderPipeline/Source/MovieRenderPipelineEditor/Private/Graph/MovieEdGraphNode.cpp
--- a/Engine/Plugins/MovieScene/MovieRenderPipeline/Source/MovieRenderPipelineEditor/Private/Graph/MovieEdGraphNode.cpp
+++ b/Engine/Plugins/MovieScene/MovieRenderPipeline/Source/MovieRenderPipelineEditor/Private/Graph/MovieEdGraphNode.cpp
@@ -259,7 +259,8 @@ FText UMoviePipelineEdGraphNode::GetNodeTitle(ENodeTitleType::Type TitleType) co
 		return RuntimeNode->GetNodeTitle(bGetDescriptive);
 	}
 
-	return LOCTEXT("UnknownNodeTitle", "Unknown");
+	// We don't know 100% that the null node is from an unknown plugin, but it's the only known reason why this could happen.
+	return LOCTEXT("UnknownNodeTitle", "Node from Unknown Plugin\nSaving the graph will result in data loss for this node!");
 }
 
 FText UMoviePipelineEdGraphNode::GetTooltipText() const
@@ -599,8 +600,8 @@ FSlateIcon UMoviePipelineEdGraphNodeBase::GetIconAndTint(FLinearColor& OutColor)
 		return RuntimeNode->GetIconAndTint(OutColor);
 	}
 
-	OutColor = FLinearColor::White;
-	return FSlateIcon();
+	OutColor = FLinearColor::Yellow;
+	return FSlateIcon(FAppStyle::GetAppStyleSetName(), "Icons.Alert");
 }
 
 bool UMoviePipelineEdGraphNodeBase::ShowPaletteIconOnNode() const
@@ -651,6 +652,13 @@ void UMoviePipelineEdGraphNodeBase::PostLoad()
 {
 	Super::PostLoad();
 
+	// This might be a node that came from a plugin that's not loaded; if that's the case, skip PostLoad().
+	if (!IsValid(RuntimeNode))
+	{
+		UE_LOG(LogMovieRenderPipeline, Warning, TEXT("Found an invalid node during PostLoad(). It's probably from a plugin which is not currently loaded."));
+		return;
+	}
+
 	// Some older nodes did not have the pin type properly set on the editor pin (specifically the value type object).
 	for (UEdGraphPin* Pin : GetAllPins())
 	{
diff --git a/Engine/Plugins/MovieScene/MovieRenderPipeline/Source/MovieRenderPipelineEditor/Private/Graph/MovieGraphSchema.cpp b/Engine/Plugins/MovieScene/MovieRenderPipeline/Source/MovieRenderPipelineEditor/Private/Graph/MovieGraphSchema.cpp
--- a/Engine/Plugins/MovieScene/MovieRenderPipeline/Source/MovieRenderPipelineEditor/Private/Graph/MovieGraphSchema.cpp
+++ b/Engine/Plugins/MovieScene/MovieRenderPipeline/Source/MovieRenderPipelineEditor/Private/Graph/MovieGraphSchema.cpp
@@ -23,6 +23,7 @@
 #include "MovieEdGraphRerouteNode.h"
 #include "MovieEdGraphVariableNode.h"
 #include "MoviePipelineEdGraphSubgraphNode.h"
+#include "MovieRenderPipelineCoreModule.h"
 
 TArray<UClass*> UMovieGraphSchema::MoviePipelineNodeClasses;
 
@@ -62,7 +63,11 @@ namespace UE::MovieGraph::Private
 		const UMoviePipelineEdGraphNodeBase* EdGraphNode = CastChecked<UMoviePipelineEdGraphNodeBase>(InPin->GetOwningNode());
 
 		UMovieGraphNode* RuntimeNode = EdGraphNode->GetRuntimeNode();
-		check(RuntimeNode);
+		if (!IsValid(RuntimeNode))
+		{
+			UE_LOG(LogMovieRenderPipeline, Warning, TEXT("Cannot find the runtime node associated with the editor node (the node may be from a plugin that's not currently loaded)."));
+			return nullptr;
+		}
 
 		return RuntimeNode;
 	}
@@ -70,6 +75,11 @@ namespace UE::MovieGraph::Private
 	UMovieGraphPin* GetGraphPinFromEdPin(const UEdGraphPin* InPin)
 	{
 		const UMovieGraphNode* GraphNode = GetGraphNodeFromEdPin(InPin);
+		if (!IsValid(GraphNode))
+		{
+			return nullptr;
+		}
+		
 		UMovieGraphPin* GraphPin = (InPin->Direction == EGPD_Input) ? GraphNode->GetInputPin(InPin->PinName) : GraphNode->GetOutputPin(InPin->PinName);
 		check(GraphPin);
 
@@ -79,6 +89,10 @@ namespace UE::MovieGraph::Private
 	UMovieGraphConfig* GetGraphFromEdPin(const UEdGraphPin* InPin)
 	{
 		const UMovieGraphNode* RuntimeNode = GetGraphNodeFromEdPin(InPin);
+		if (!IsValid(RuntimeNode))
+		{
+			return nullptr;
+		}
 
 		UMovieGraphConfig* RuntimeGraph = RuntimeNode->GetGraph();
 		check(RuntimeGraph);
@@ -459,32 +473,43 @@ const FPinConnectionResponse UMovieGraphSchema::CanCreateConnection(const UEdGra
 	const UMovieGraphPin* FromPin = UE::MovieGraph::Private::GetGraphPinFromEdPin(PinA);
 	const UMovieGraphPin* ToPin = UE::MovieGraph::Private::GetGraphPinFromEdPin(PinB);
 	
+	if (!IsValid(FromPin) || !IsValid(ToPin))
+	{
+		return FPinConnectionResponse(CONNECT_RESPONSE_DISALLOW, LOCTEXT("InvalidNodeError", "The to/from pin/node is invalid!"));
+	}
+	
 	return FromPin->CanCreateConnection_PinConnectionResponse(ToPin);
 }
 
 bool UMovieGraphSchema::TryCreateConnection(UEdGraphPin* InA, UEdGraphPin* InB) const
 {
-	// See if the native UEdGraph connection goes through
-	const bool bModified = Super::TryCreateConnection(InA, InB);
+	check(InA && InB);
+	const UEdGraphPin* A = (InA->Direction == EGPD_Output) ? InA : InB;
+	const UEdGraphPin* B = (InA->Direction == EGPD_Input) ? InA : InB;
+	check(A->Direction == EGPD_Output && B->Direction == EGPD_Input);
 
-	// If it does, try to propagate the change to our runtime graph
-	if (bModified)
-	{
-		check(InA && InB);
-		const UEdGraphPin* A = (InA->Direction == EGPD_Output) ? InA : InB;
-		const UEdGraphPin* B = (InA->Direction == EGPD_Input) ? InA : InB;
-		check(A->Direction == EGPD_Output && B->Direction == EGPD_Input);
+	const UMoviePipelineEdGraphNodeBase* EdGraphNodeA = CastChecked<UMoviePipelineEdGraphNodeBase>(A->GetOwningNode());
+	const UMoviePipelineEdGraphNodeBase* EdGraphNodeB = CastChecked<UMoviePipelineEdGraphNodeBase>(B->GetOwningNode());
 
-		UMoviePipelineEdGraphNodeBase* EdGraphNodeA = CastChecked<UMoviePipelineEdGraphNodeBase>(A->GetOwningNode());
-		UMoviePipelineEdGraphNodeBase* EdGraphNodeB = CastChecked<UMoviePipelineEdGraphNodeBase>(B->GetOwningNode());
+	UMovieGraphNode* RuntimeNodeA = EdGraphNodeA->GetRuntimeNode();
+	UMovieGraphNode* RuntimeNodeB = EdGraphNodeB->GetRuntimeNode();
 
-		UMovieGraphNode* RuntimeNodeA = EdGraphNodeA->GetRuntimeNode();
-		UMovieGraphNode* RuntimeNodeB = EdGraphNodeB->GetRuntimeNode();
-		check(RuntimeNodeA && RuntimeNodeB);
+	// If the node associated with either of the pins is invalid, the node is probably from a plugin that isn't loaded. If this is the case, bail
+	// on creating the connection.
+	if (!IsValid(RuntimeNodeA) || !IsValid(RuntimeNodeB))
+	{
+		UE_LOG(LogMovieRenderPipeline, Warning, TEXT("Cannot create a connection to/from a node which is not currently valid (it may be from a plugin that's not currently loaded)."));
+		return false;
+	}
 
-		UMovieGraphConfig* RuntimeGraph = RuntimeNodeA->GetGraph();
-		check(RuntimeGraph);
+	UMovieGraphConfig* RuntimeGraph = RuntimeNodeA->GetGraph();
+	check(RuntimeGraph);
 
+	// See if the native UEdGraph connection goes through.
+	// If the connection was made, try to propagate the change to our runtime graph.
+	const bool bModified = Super::TryCreateConnection(InA, InB);
+	if (bModified)
+	{
 		const bool bReconstructNodeB = RuntimeGraph->AddLabeledEdge(RuntimeNodeA, A->PinName, RuntimeNodeB, B->PinName);
 		//if (bReconstructNodeB)
 		//{
@@ -503,8 +528,14 @@ void UMovieGraphSchema::BreakPinLinks(UEdGraphPin& TargetPin, bool bSendsNodeNot
 	UEdGraphNode* GraphNode = TargetPin.GetOwningNode();
 	UMoviePipelineEdGraphNodeBase* MoviePipelineEdGraphNode = CastChecked<UMoviePipelineEdGraphNodeBase>(GraphNode);
 
+	// The node may be invalid if it's from a plugin that isn't currently loaded. Skip breaking the connection, although that's not ideal because the
+	// user is probably trying to get rid of the broken node.
 	UMovieGraphNode* RuntimeNode = MoviePipelineEdGraphNode->GetRuntimeNode();
-	check(RuntimeNode);
+	if (!IsValid(RuntimeNode))
+	{
+		UE_LOG(LogMovieRenderPipeline, Warning, TEXT("Cannot remove connection from node which is not currently valid (it may be from a plugin that's not currently loaded)."));
+		return;
+	}
 
 	UMovieGraphConfig* RuntimeGraph = RuntimeNode->GetGraph();
 	check(RuntimeGraph);
