MoviePipeline: Fixed some more crashes related to graph evaluation and graph interaction when the graph contains a node that is from an unloaded plugin. Also updated the "Unknown" node to include the name of the (unloaded) plugin that the node originated from.

#rnx

diff --git a/Engine/Plugins/MovieScene/MovieRenderPipeline/Source/MovieRenderPipelineCore/Private/Graph/MovieGraphConfig.cpp b/Engine/Plugins/MovieScene/MovieRenderPipeline/Source/MovieRenderPipelineCore/Private/Graph/MovieGraphConfig.cpp
--- a/Engine/Plugins/MovieScene/MovieRenderPipeline/Source/MovieRenderPipelineCore/Private/Graph/MovieGraphConfig.cpp
+++ b/Engine/Plugins/MovieScene/MovieRenderPipeline/Source/MovieRenderPipelineCore/Private/Graph/MovieGraphConfig.cpp
@@ -1984,6 +1984,11 @@ void UMovieGraphConfig::VisitUpstreamNodes_Recursive(UMovieGraphNode* FromNode,
 	{
 		for (const UMovieGraphPin* ConnectedPin : Pin->GetAllConnectedPins())
 		{
+			if (!ConnectedPin)
+			{
+				continue;
+			}
+
 			// Branches are followed during traversal, but we also want to traverse reroute nodes as well since they're passthrough
 			if (ConnectedPin->Properties.bIsBranch || (Pin->Node && Pin->Node->IsA<UMovieGraphRerouteNode>()))
 			{
@@ -2021,6 +2026,11 @@ void UMovieGraphConfig::VisitDownstreamNodes_Recursive(UMovieGraphNode* FromNode
 	{
 		for (const UMovieGraphPin* ConnectedPin : Pin->GetAllConnectedPins())
 		{
+			if (!ConnectedPin)
+			{
+				continue;
+			}
+
 			// Branches are followed during traversal, but we also want to traverse reroute nodes as well since they're passthrough
 			if (ConnectedPin->Properties.bIsBranch || (Pin->Node && Pin->Node->IsA<UMovieGraphRerouteNode>()))
 			{
@@ -2079,6 +2089,14 @@ UMovieGraphEvaluatedConfig* UMovieGraphConfig::CreateFlattenedGraph(const FMovie
 			// Follow the branch connected to this pin
 			for (UMovieGraphEdge* Edge : InputPin->Edges)
 			{
+				if (!IsValid(Edge))
+				{
+					const FText InvalidNodeError = FText::Format(LOCTEXT("TraversalError_InvalidNode", "Found an invalid node in branch [{0}]. Aborting graph evaluation."), { FText::FromName(BranchName) });
+					UE_LOG(LogMovieRenderPipeline, Error, TEXT("%s"), *InvalidNodeError.ToString());
+					OutError = InvalidNodeError.ToString();
+					return nullptr;
+				}
+				
 				UMovieGraphPin* OtherPin = Edge->GetOtherPin(InputPin);
 				if (OtherPin)
 				{
@@ -2103,6 +2121,11 @@ UMovieGraphEvaluatedConfig* UMovieGraphConfig::CreateFlattenedGraph(const FMovie
 				GlobalStackContext.UserContext = InContext;
 				for (UMovieGraphEdge* Edge : GlobalsPin->Edges)
 				{
+					if (!IsValid(Edge))
+					{
+						continue;
+					}
+					
 					UMovieGraphPin* OtherPin = Edge->GetOtherPin(GlobalsPin);
 					if (OtherPin)
 					{
diff --git a/Engine/Plugins/MovieScene/MovieRenderPipeline/Source/MovieRenderPipelineCore/Private/Graph/MovieGraphPin.cpp b/Engine/Plugins/MovieScene/MovieRenderPipeline/Source/MovieRenderPipelineCore/Private/Graph/MovieGraphPin.cpp
--- a/Engine/Plugins/MovieScene/MovieRenderPipeline/Source/MovieRenderPipelineCore/Private/Graph/MovieGraphPin.cpp
+++ b/Engine/Plugins/MovieScene/MovieRenderPipeline/Source/MovieRenderPipelineCore/Private/Graph/MovieGraphPin.cpp
@@ -99,6 +99,11 @@ bool UMovieGraphPin::BreakAllEdges()
 
 	for (UMovieGraphEdge* Edge : Edges)
 	{
+		if (!Edge)
+		{
+			continue;
+		}
+
 		if (UMovieGraphPin* OtherPin = Edge->GetOtherPin(this))
 		{
 			OtherPin->Modify();
@@ -198,16 +203,8 @@ bool UMovieGraphPin::IsOutputPin() const
 
 int32 UMovieGraphPin::EdgeCount() const
 {
-	int32 EdgeCount = 0;
-	for (UMovieGraphEdge* Edge : Edges)
-	{
-		if (Edge->IsValid())
-		{
-			EdgeCount++;
-		}
-	}
-
-	return EdgeCount;
+	// Includes invalid edges, if any
+	return Edges.Num();
 }
 
 bool UMovieGraphPin::AllowsMultipleConnections() const
diff --git a/Engine/Plugins/MovieScene/MovieRenderPipeline/Source/MovieRenderPipelineCore/Private/Graph/MovieGraphPipeline.cpp b/Engine/Plugins/MovieScene/MovieRenderPipeline/Source/MovieRenderPipelineCore/Private/Graph/MovieGraphPipeline.cpp
--- a/Engine/Plugins/MovieScene/MovieRenderPipeline/Source/MovieRenderPipelineCore/Private/Graph/MovieGraphPipeline.cpp
+++ b/Engine/Plugins/MovieScene/MovieRenderPipeline/Source/MovieRenderPipelineCore/Private/Graph/MovieGraphPipeline.cpp
@@ -77,7 +77,6 @@ void UMovieGraphPipeline::Initialize(UMoviePipelineExecutorJob* InJob, const FMo
 		return;
 	}
 
-
 	// Register any additional engine callbacks needed.
 	{
 		// Called before the Custom Timestep is updated. This gives us a chance to calculate
@@ -87,6 +86,7 @@ void UMovieGraphPipeline::Initialize(UMoviePipelineExecutorJob* InJob, const FMo
 		// Called at the end of the frame after everything has been ticked and rendered for the frame.
 		FCoreDelegates::OnEndFrame.AddUObject(this, &UMovieGraphPipeline::OnEngineTickEndFrame);
 	}
+
 	CurrentJob = InJob;
 	CurrentShotIndex = 0;
 	GraphInitializationTime = FDateTime::UtcNow();
@@ -96,11 +96,12 @@ void UMovieGraphPipeline::Initialize(UMoviePipelineExecutorJob* InJob, const FMo
 	DuplicateJobAndConfiguration();
 	ExecutePreJobScripts();
 
-	// Lay groundwork for Unreal Insights
 	const FMovieGraphTraversalContext Context = GetCurrentTraversalContext(false);
 	FString OutError;
+	UMovieGraphEvaluatedConfig* FlattenedConfig = GetCurrentJob()->GetGraphPreset()->CreateFlattenedGraph(Context, OutError);
 
-	if (UMovieGraphEvaluatedConfig* FlattenedConfig = GetCurrentJob()->GetGraphPreset()->CreateFlattenedGraph(Context, OutError))
+	// Lay groundwork for Unreal Insights
+	if (FlattenedConfig)
 	{
 		constexpr bool bIncludeCDOs = false;
 		constexpr bool bExactMatch = true;
@@ -111,13 +112,17 @@ void UMovieGraphPipeline::Initialize(UMoviePipelineExecutorJob* InJob, const FMo
 			StartUnrealInsightsCapture(FlattenedConfig);
 		}
 	}
+	else
+	{
+		// Although the evaluated graph was created for Unreal Insights, it will lead to a general failure to initialize the pipeline.
+		UE_LOG(LogMovieRenderPipeline, Error, TEXT("Unable to initialize Movie Pipeline. The following error was encountered: %s"), *OutError);
+	}
 
 	// Create instances of our different classes from the InitConfig
 	GraphRendererInstance = NewObject<UMovieGraphRendererBase>(this, InitConfig.RendererClass);
 	GraphDataSourceInstance = NewObject<UMovieGraphDataSourceBase>(this, InitConfig.DataSourceClass);
 	GraphAudioRendererInstance = NewObject<UMovieGraphAudioRendererBase>(this, InitConfig.AudioRendererClass);
 
-
 	// Now that we've created our various systems, we will start using them. First thing we do is cache data about
 	// the world, job, player viewport, etc, before we make any modifications. These will be restored at the end
 	// of the render.
@@ -128,7 +133,6 @@ void UMovieGraphPipeline::Initialize(UMoviePipelineExecutorJob* InJob, const FMo
 
 	GraphAudioRendererInstance->SetupAudioRendering();
 
-
 	// Update our list of shots from our data source, and then
 	// create our list of active shots, so we don't try to render
 	// a shot the user has deactivated.
@@ -149,13 +153,17 @@ void UMovieGraphPipeline::Initialize(UMoviePipelineExecutorJob* InJob, const FMo
 	// ToDo: Print some information, Job Name, Level Sequence, etc.
 	UE_LOG(LogMovieRenderPipeline, Log, TEXT("Finished initializing MovieGraph Render"));
 
+	const bool bNoShotsToRender = ActiveShotList.Num() == 0;
+	if (bNoShotsToRender)
+	{
+		UE_LOG(LogMovieRenderPipeline, Warning, TEXT("No shots detected to render. Either all outside playback range, or disabled via shot mask. Shutting down."));
+	}
 
-	// If the shot mask entirely disabled everything we'll transition directly to finish as there is no work to do.
-	if (ActiveShotList.Num() == 0)
+	// If shutting down immediately, we'll transition directly to finish as there is no work to do.
+	// If the flattened config isn't valid at this point, evaluation failed, and a shutdown is required.
+	if (bNoShotsToRender || !FlattenedConfig)
 	{
 		// We have to transition twice as Uninitialized -> n state is a no-op, so the second tick will take us to Finished which shuts down.
-		UE_LOG(LogMovieRenderPipeline, Warning, TEXT("MovieGraph No shots detected to render. Either all outside playback range, or disabled via shot mask. Shutting down."));
-	
 		TransitionToState(EMovieRenderPipelineState::Export);
 		TransitionToState(EMovieRenderPipelineState::Finished);
 	}
diff --git a/Engine/Plugins/MovieScene/MovieRenderPipeline/Source/MovieRenderPipelineCore/Private/Graph/MovieGraphSequenceDataSource.cpp b/Engine/Plugins/MovieScene/MovieRenderPipeline/Source/MovieRenderPipelineCore/Private/Graph/MovieGraphSequenceDataSource.cpp
--- a/Engine/Plugins/MovieScene/MovieRenderPipeline/Source/MovieRenderPipelineCore/Private/Graph/MovieGraphSequenceDataSource.cpp
+++ b/Engine/Plugins/MovieScene/MovieRenderPipeline/Source/MovieRenderPipelineCore/Private/Graph/MovieGraphSequenceDataSource.cpp
@@ -68,8 +68,13 @@ void UMovieGraphSequenceDataSource::OverrideSequencePlaybackRangeFromGlobalOutpu
 	FMovieGraphTraversalContext TraversalContext;
 	UMoviePipelineExecutorJob* CurrentJob = GetOwningGraph()->GetCurrentJob();
 	TraversalContext.Job = CurrentJob;
+
 	FString OutErrorMessage;
 	UMovieGraphEvaluatedConfig* EvaluatedGraph = CurrentJob->GetGraphPreset()->CreateFlattenedGraph(TraversalContext, OutErrorMessage);
+	if (!EvaluatedGraph)
+	{
+		return;
+	}
 
 	constexpr bool bIncludeCDOs = true;
 	constexpr bool bExactMatch = false;
diff --git a/Engine/Plugins/MovieScene/MovieRenderPipeline/Source/MovieRenderPipelineCore/Public/Graph/MovieGraphPin.h b/Engine/Plugins/MovieScene/MovieRenderPipeline/Source/MovieRenderPipelineCore/Public/Graph/MovieGraphPin.h
--- a/Engine/Plugins/MovieScene/MovieRenderPipeline/Source/MovieRenderPipelineCore/Public/Graph/MovieGraphPin.h
+++ b/Engine/Plugins/MovieScene/MovieRenderPipeline/Source/MovieRenderPipelineCore/Public/Graph/MovieGraphPin.h
@@ -129,6 +129,7 @@ public:
 	UFUNCTION(BlueprintCallable, Category = "Movie Graph")
 	bool IsOutputPin() const;
 
+	/** Gets the number of edges connected to this pin. May include invalid edges (which can happen if connected to an invalid node). */
 	UFUNCTION(BlueprintCallable, Category = "Movie Graph")
 	int32 EdgeCount() const;
 
diff --git a/Engine/Plugins/MovieScene/MovieRenderPipeline/Source/MovieRenderPipelineEditor/MovieRenderPipelineEditor.Build.cs b/Engine/Plugins/MovieScene/MovieRenderPipeline/Source/MovieRenderPipelineEditor/MovieRenderPipelineEditor.Build.cs
--- a/Engine/Plugins/MovieScene/MovieRenderPipeline/Source/MovieRenderPipelineEditor/MovieRenderPipelineEditor.Build.cs
+++ b/Engine/Plugins/MovieScene/MovieRenderPipeline/Source/MovieRenderPipelineEditor/MovieRenderPipelineEditor.Build.cs
@@ -75,6 +75,7 @@ public class MovieRenderPipelineEditor : ModuleRules
 				"OpenColorIOEditor",
 				"Imath",		// This and UEOpenExr are only required for use of UMovieGraphImageSequenceOutputNode. Ideally these are refactored
 				"UEOpenExr",	// out of the ImageSequenceOutputNode header at some point so these dependencies can be removed.
+				"Projects",		// For IPluginManager
 			}
         );
 
diff --git a/Engine/Plugins/MovieScene/MovieRenderPipeline/Source/MovieRenderPipelineEditor/Private/Graph/MovieEdGraph.cpp b/Engine/Plugins/MovieScene/MovieRenderPipeline/Source/MovieRenderPipelineEditor/Private/Graph/MovieEdGraph.cpp
--- a/Engine/Plugins/MovieScene/MovieRenderPipeline/Source/MovieRenderPipelineEditor/Private/Graph/MovieEdGraph.cpp
+++ b/Engine/Plugins/MovieScene/MovieRenderPipeline/Source/MovieRenderPipelineEditor/Private/Graph/MovieEdGraph.cpp
@@ -170,7 +170,7 @@ void UMoviePipelineEdGraph::CreateLinks(UMoviePipelineEdGraphNodeBase* InGraphNo
 
 			for (const UMovieGraphEdge* Edge : Pin->Edges)
 			{
-				if (!Edge->IsValid())
+				if (!Edge || !Edge->IsValid())
 				{
 					UE_LOG(LogMovieRenderPipeline, Error, TEXT("Invalid edge for Pin: %s"), *Pin->Properties.Label.ToString());
 					continue;
diff --git a/Engine/Plugins/MovieScene/MovieRenderPipeline/Source/MovieRenderPipelineEditor/Private/Graph/MovieEdGraphNode.cpp b/Engine/Plugins/MovieScene/MovieRenderPipeline/Source/MovieRenderPipelineEditor/Private/Graph/MovieEdGraphNode.cpp
--- a/Engine/Plugins/MovieScene/MovieRenderPipeline/Source/MovieRenderPipelineEditor/Private/Graph/MovieEdGraphNode.cpp
+++ b/Engine/Plugins/MovieScene/MovieRenderPipeline/Source/MovieRenderPipelineEditor/Private/Graph/MovieEdGraphNode.cpp
@@ -7,6 +7,7 @@
 #include "Graph/MovieGraphConfig.h"
 #include "Graph/Nodes/MovieGraphVariableNode.h"
 #include "EdGraph/EdGraphPin.h"
+#include "Interfaces/IPluginManager.h"
 #include "Misc/TransactionObjectEvent.h"
 #include "MovieEdGraph.h"
 #include "MovieGraphSchema.h"
@@ -260,7 +261,9 @@ FText UMoviePipelineEdGraphNode::GetNodeTitle(ENodeTitleType::Type TitleType) co
 	}
 
 	// We don't know 100% that the null node is from an unknown plugin, but it's the only known reason why this could happen.
-	return LOCTEXT("UnknownNodeTitle", "Node from Unknown Plugin\nSaving the graph will result in data loss for this node!");
+	const FText PluginName = OriginPluginName.IsEmpty() ? LOCTEXT("UnknownPluginName", "Unknown") : FText::FromString(OriginPluginName);
+	const FText UnknownPluginNodeTitle = FText::Format(LOCTEXT("NodeTitle_UnknownPlugin", "Node from Unknown Plugin [{0}]\nSaving the graph will result in data loss for this node!\nCan be disconnected, but not deleted."), { PluginName });
+	return UnknownPluginNodeTitle;
 }
 
 FText UMoviePipelineEdGraphNode::GetTooltipText() const
@@ -655,10 +658,21 @@ void UMoviePipelineEdGraphNodeBase::PostLoad()
 	// This might be a node that came from a plugin that's not loaded; if that's the case, skip PostLoad().
 	if (!IsValid(RuntimeNode))
 	{
-		UE_LOG(LogMovieRenderPipeline, Warning, TEXT("Found an invalid node during PostLoad(). It's probably from a plugin which is not currently loaded."));
+		// This should really be a warning, but currently you cannot delete an invalid node. So even if a user bypasses/disconnects an invalid node,
+		// they'll see this warning every time they render, which isn't ideal. There will be plenty of other warnings/errors that show up if they
+		// do not bypass the node.
+		UE_LOG(LogMovieRenderPipeline, Log, TEXT("Found an invalid node during PostLoad(). It's probably from a plugin which is not currently loaded."));
 		return;
 	}
 
+	// If the plugin this node came from can be determined, cache it. If the plugin is unloaded at some point, we can display this cached plugin name
+	// in the UI as a hint.
+	const FString NodePackageName = FPackageName::GetShortName(RuntimeNode->GetClass()->GetPackage());
+	if (const TSharedPtr<IPlugin> Plugin = IPluginManager::Get().GetModuleOwnerPlugin(FName(NodePackageName)))
+	{
+		OriginPluginName = Plugin->GetFriendlyName();
+	}
+
 	// Some older nodes did not have the pin type properly set on the editor pin (specifically the value type object).
 	for (UEdGraphPin* Pin : GetAllPins())
 	{
@@ -686,6 +700,12 @@ void UMoviePipelineEdGraphNodeBase::ReconstructNode()
 		return;
 	}
 
+	// Also don't reconstruct if the runtime node is nullptr. This is most likely because the node is from an unloaded plugin.
+	if (!GetRuntimeNode())
+	{
+		return;
+	}
+
 	UMoviePipelineEdGraph* Graph = CastChecked<UMoviePipelineEdGraph>(GetGraph());
 	Graph->Modify();
 
@@ -831,11 +851,20 @@ void UMoviePipelineEdGraphNodeBase::GetNodeContextMenuActions(UToolMenu* Menu, U
 
 	{
 		FToolMenuSection& Section = Menu->AddSection("EdGraphSchemaNodeActions", LOCTEXT("NodeActionsHeader", "Node Actions"));
-		Section.AddMenuEntry(FGraphEditorCommands::Get().BreakNodeLinks);
-		Section.AddMenuEntry(FGenericCommands::Get().Delete);
-		Section.AddMenuEntry(FGenericCommands::Get().Cut);
-		Section.AddMenuEntry(FGenericCommands::Get().Copy);
-		Section.AddMenuEntry(FGenericCommands::Get().Duplicate);
+
+		// The runtime node must be valid in order to be manipulated. If it's nullptr (likely from an unloaded plugin) then it cannot be changed.
+		if (const UMoviePipelineEdGraphNodeBase* GraphNode = Cast<const UMoviePipelineEdGraphNodeBase>(Context->Node))
+		{
+			if (GraphNode->GetRuntimeNode())
+			{
+				Section.AddMenuEntry(FGraphEditorCommands::Get().BreakNodeLinks);
+				Section.AddMenuEntry(FGenericCommands::Get().Delete);
+				Section.AddMenuEntry(FGenericCommands::Get().Cut);
+				Section.AddMenuEntry(FGenericCommands::Get().Copy);
+				Section.AddMenuEntry(FGenericCommands::Get().Duplicate);
+			}
+		}
+
 		Section.AddMenuEntry(FGraphEditorCommands::Get().EnableNodes);
 		Section.AddMenuEntry(FGraphEditorCommands::Get().DisableNodes);
 	}
diff --git a/Engine/Plugins/MovieScene/MovieRenderPipeline/Source/MovieRenderPipelineEditor/Private/Graph/MovieEdGraphNode.h b/Engine/Plugins/MovieScene/MovieRenderPipeline/Source/MovieRenderPipelineEditor/Private/Graph/MovieEdGraphNode.h
--- a/Engine/Plugins/MovieScene/MovieRenderPipeline/Source/MovieRenderPipelineEditor/Private/Graph/MovieEdGraphNode.h
+++ b/Engine/Plugins/MovieScene/MovieRenderPipeline/Source/MovieRenderPipelineEditor/Private/Graph/MovieEdGraphNode.h
@@ -80,6 +80,13 @@ protected:
 	UPROPERTY()
 	TObjectPtr<UMovieGraphNode> RuntimeNode;
 
+	/**
+	 * Cached friendly name of the plugin that this node originated from. In the event that the origin plugin is unloaded, we can provide the name of
+	 * the plugin as a hint. The runtime node will be null in this scenario, so specifying the name of the plugin could be very helpful.
+	 */
+	UPROPERTY()
+	FString OriginPluginName;
+
 	/** Should we early out in ReconstructNode(), skipping restoring the connections to other nodes. Set during Copy/Paste. */
 	bool bDisableReconstructNode = false;
 
diff --git a/Engine/Plugins/MovieScene/MovieRenderPipeline/Source/MovieRenderPipelineEditor/Private/Widgets/Graph/SMovieGraphConfigPanel.cpp b/Engine/Plugins/MovieScene/MovieRenderPipeline/Source/MovieRenderPipelineEditor/Private/Widgets/Graph/SMovieGraphConfigPanel.cpp
--- a/Engine/Plugins/MovieScene/MovieRenderPipeline/Source/MovieRenderPipelineEditor/Private/Widgets/Graph/SMovieGraphConfigPanel.cpp
+++ b/Engine/Plugins/MovieScene/MovieRenderPipeline/Source/MovieRenderPipelineEditor/Private/Widgets/Graph/SMovieGraphConfigPanel.cpp
@@ -269,8 +269,14 @@ bool SMoviePipelineGraphPanel::CanDeleteSelectedNodes() const
 
 	for (UObject* Object : GraphEditorWidget->GetSelectedNodes())
 	{
-		UEdGraphNode* GraphNode = Cast<UEdGraphNode>(Object);
-		if (GraphNode && !GraphNode->CanUserDeleteNode())
+		const UMoviePipelineEdGraphNodeBase* GraphNode = Cast<UMoviePipelineEdGraphNodeBase>(Object);
+		if (!GraphNode || !GraphNode->CanUserDeleteNode())
+		{
+			return false;
+		}
+
+		// Additionally, if the associated runtime node is not valid (likely from an unloaded plugin), then the node cannot be deleted.
+		if (!GraphNode->GetRuntimeNode())
 		{
 			return false;
 		}
@@ -371,8 +377,14 @@ bool SMoviePipelineGraphPanel::CanCopySelectedNodes() const
 	{
 		for (UObject* Object : GraphEditorWidget->GetSelectedNodes())
 		{
-			if (UEdGraphNode* GraphNode = CastChecked<UEdGraphNode>(Object))
+			if (const UMoviePipelineEdGraphNodeBase* GraphNode = CastChecked<UMoviePipelineEdGraphNodeBase>(Object))
 			{
+				// Invalid runtime nodes (likely from an unloaded plugin) cannot be copied.
+				if (!GraphNode->GetRuntimeNode())
+				{
+					return false;
+				}
+				
 				if (GraphNode->CanDuplicateNode())
 				{
 					bAtLeastOneCopyableNode = true;
