Allow landscape grass culling when running MRQ.
Previously we would pass an array of no camera to RegenerateGrass, which prevents all landscape grass from being distance-culled. Instead, we now pass nothing, which leads landscape to use the streaming manager to get the current "views", one of which is the captured camera (the other, the editor). This is better handled in the new movie render graph but at least this allows solving a performance issue in grass-heavy scenes.
The user can now override the grass cull distance (as well as grass density) via Game Overrides settings.
Default = override + cull distance multiplier of 50, to have (mostly) backwards compatibility with the previous MRQ captures (which prevented all grass culling)
The default for grass density override is false, though, also for backwards compatibility reasons.

Note that the new movie render graph doesn't have these settings since it doesn't have them for ViewDistanceScale, ShadowDistanceScale, etc. so I figured it didn't need those grass settings either

Also added a CVar (grass.VisualLog.ShowCameraLocation) to show the camera locations used by landscape grass in the visual logger

#tests Editor, MRQ

diff --git a/Engine/Plugins/MovieScene/MovieRenderPipeline/Source/MovieRenderPipelineCore/Private/MoviePipelineGameOverrideSetting.cpp b/Engine/Plugins/MovieScene/MovieRenderPipeline/Source/MovieRenderPipelineCore/Private/MoviePipelineGameOverrideSetting.cpp
--- a/Engine/Plugins/MovieScene/MovieRenderPipeline/Source/MovieRenderPipelineCore/Private/MoviePipelineGameOverrideSetting.cpp
+++ b/Engine/Plugins/MovieScene/MovieRenderPipeline/Source/MovieRenderPipelineCore/Private/MoviePipelineGameOverrideSetting.cpp
@@ -111,6 +111,16 @@ void UMoviePipelineGameOverrideSetting::ApplyCVarSettings(const bool bOverrideVa
 		MOVIEPIPELINE_STORE_AND_OVERRIDE_CVAR_FLOAT(PreviousViewDistanceScale, TEXT("r.ViewDistanceScale"), ViewDistanceScale, bOverrideValues);
 	}
 
+	if (bOverrideGrassCullDistanceScale)
+	{
+		MOVIEPIPELINE_STORE_AND_OVERRIDE_CVAR_FLOAT(PreviousGrassCullDistanceScale, TEXT("grass.CullDistanceScale"), GrassCullDistanceScale, bOverrideValues);
+	}
+
+	if (bOverrideGrassDensityScale)
+	{
+		MOVIEPIPELINE_STORE_AND_OVERRIDE_CVAR_FLOAT(PreviousGrassDensityScale, TEXT("grass.densityScale"), GrassDensityScale, bOverrideValues);
+	}
+
 	if (bDisableGPUTimeout)
 	{
 		// This CVAR only exists if the D3D12RHI module is loaded
@@ -243,6 +253,16 @@ void UMoviePipelineGameOverrideSetting::BuildNewProcessCommandLineArgsImpl(TArra
 		InOutDeviceProfileCvars.AddUnique(FString::Printf(TEXT("r.ViewDistanceScale=%d"), ViewDistanceScale));
 	}
 
+	if (bOverrideGrassCullDistanceScale)
+	{
+		InOutDeviceProfileCvars.AddUnique(FString::Printf(TEXT("grass.CullDistanceScale=%f"), GrassCullDistanceScale));
+	}
+
+	if (bOverrideGrassDensityScale)
+	{
+		InOutDeviceProfileCvars.AddUnique(FString::Printf(TEXT("grass.densityScale=%f"), GrassDensityScale));
+	}
+
 	if (bDisableGPUTimeout)
 	{
 		InOutDeviceProfileCvars.AddUnique(TEXT("r.D3D12.GPUTimeout=0"));
diff --git a/Engine/Plugins/MovieScene/MovieRenderPipeline/Source/MovieRenderPipelineCore/Private/MoviePipelineRendering.cpp b/Engine/Plugins/MovieScene/MovieRenderPipeline/Source/MovieRenderPipelineCore/Private/MoviePipelineRendering.cpp
--- a/Engine/Plugins/MovieScene/MovieRenderPipeline/Source/MovieRenderPipelineCore/Private/MoviePipelineRendering.cpp
+++ b/Engine/Plugins/MovieScene/MovieRenderPipeline/Source/MovieRenderPipelineCore/Private/MoviePipelineRendering.cpp
@@ -666,7 +666,7 @@ void UMoviePipeline::FlushAsyncEngineSystems()
 		{
 			if (ULandscapeSubsystem* LandscapeSubsystem = GetWorld()->GetSubsystem<ULandscapeSubsystem>())
 			{
-				LandscapeSubsystem->RegenerateGrass(/*bInFlushGrass = */false, /*bInForceSync = */true, /*InOptionalCameraLocations = */MakeArrayView(TArray<FVector>()));
+				LandscapeSubsystem->RegenerateGrass(/*bInFlushGrass = */false, /*bInForceSync = */true);
 			}
 		}
 	}
diff --git a/Engine/Plugins/MovieScene/MovieRenderPipeline/Source/MovieRenderPipelineCore/Public/MoviePipelineGameOverrideSetting.h b/Engine/Plugins/MovieScene/MovieRenderPipeline/Source/MovieRenderPipelineCore/Public/MoviePipelineGameOverrideSetting.h
--- a/Engine/Plugins/MovieScene/MovieRenderPipeline/Source/MovieRenderPipelineCore/Public/MoviePipelineGameOverrideSetting.h
+++ b/Engine/Plugins/MovieScene/MovieRenderPipeline/Source/MovieRenderPipelineCore/Public/MoviePipelineGameOverrideSetting.h
@@ -39,6 +39,10 @@ public:
 		, bOverrideViewDistanceScale(true)
 		, ViewDistanceScale(50)
 		, bFlushGrassStreaming(true)
+		, bOverrideGrassCullDistanceScale(true)
+		, GrassCullDistanceScale(50.0f)
+		, bOverrideGrassDensityScale(false)
+		, GrassDensityScale(1.0f)
 		, bFlushStreamingManagers(true)
 		, bOverrideVirtualTextureFeedbackFactor(true)
 		, VirtualTextureFeedbackFactor(1)
@@ -104,17 +108,33 @@ public:
 	float ShadowRadiusThreshold;
 
 	/** Should we override the View Distance Scale? Can be used in situations where MaxDrawDistance has been set before for in-game performance. */
-	UPROPERTY(EditAnywhere, BlueprintReadWrite, Category = "Rendering")
+	UPROPERTY(EditAnywhere, BlueprintReadWrite, Category = "Rendering", meta = (InlineEditConditionToggle))
 	bool bOverrideViewDistanceScale;
 
 	/** Controls the view distance scale. A primitive's MaxDrawDistance is scaled by this value. */
 	UPROPERTY(EditAnywhere, BlueprintReadWrite, Category = "Rendering", meta = (EditCondition = bOverrideViewDistanceScale))
 	int32 ViewDistanceScale;
 	
-	/** Flushing grass streaming (combined with override view distance scale) prevents visible pop-in/culling of grace instances. */
+	/** Flushing grass streaming (combined with override view distance scale) prevents visible pop-in/culling of grass instances. */
 	UPROPERTY(EditAnywhere, BlueprintReadWrite, Category = "Rendering")
 	bool bFlushGrassStreaming;
 
+	/** Should we override the Landscape Grass Cull Distance Scale? Can be used in situations where grass.CullDistanceScale has been set before for in-game performance. */
+	UPROPERTY(EditAnywhere, BlueprintReadWrite, Category = "Rendering", meta = (InlineEditConditionToggle))
+	bool bOverrideGrassCullDistanceScale;
+
+	/** Controls the grass cull distance scale. See grass.CullDistanceScale */
+	UPROPERTY(EditAnywhere, BlueprintReadWrite, Category = "Rendering", meta = (EditCondition = bOverrideGrassCullDistanceScale))
+	float GrassCullDistanceScale;
+
+	/** Should we override the Landscape Grass Density Scale? Can be used in situations where grass.densityScale has been set before for in-game performance. */
+	UPROPERTY(EditAnywhere, BlueprintReadWrite, Category = "Rendering", meta = (InlineEditConditionToggle))
+	bool bOverrideGrassDensityScale;
+
+	/** Controls the grass density scale. See grass.densityScale */
+	UPROPERTY(EditAnywhere, BlueprintReadWrite, Category = "Rendering", meta = (EditCondition = bOverrideGrassDensityScale))
+	float GrassDensityScale;
+
 	/** Experimental. If true flush the streaming managers (Texture Streaming) each frame. Allows Texture Streaming to not have visible pop-in in final frames. */
 	UPROPERTY(EditAnywhere, BlueprintReadWrite, Category = "Rendering")
 	bool bFlushStreamingManagers;
@@ -143,6 +163,8 @@ private:
 	int32 PreviousShadowQuality;
 	float PreviousShadowRadiusThreshold;
 	float PreviousViewDistanceScale;
+	float PreviousGrassCullDistanceScale;
+	float PreviousGrassDensityScale;
 	int32 PreviousGPUTimeout;
 	int32 PreviousAnimationUROEnabled;
 	int32 PreviousFoliageDitheredLOD;
diff --git a/Engine/Source/Runtime/Landscape/Private/LandscapeSubsystem.cpp b/Engine/Source/Runtime/Landscape/Private/LandscapeSubsystem.cpp
--- a/Engine/Source/Runtime/Landscape/Private/LandscapeSubsystem.cpp
+++ b/Engine/Source/Runtime/Landscape/Private/LandscapeSubsystem.cpp
@@ -41,6 +41,7 @@
 #include "Misc/App.h"
 #include "Misc/DateTime.h"
 #include "AssetCompilingManager.h"
+#include "VisualLogger/VisualLogger.h"
 
 #include UE_INLINE_GENERATED_CPP_BY_NAME(LandscapeSubsystem)
 
@@ -67,6 +68,11 @@ static FAutoConsoleVariableRef CVarUseStreamingManagerForCameras(
 	GUseStreamingManagerForCameras,
 	TEXT("1: Use Streaming Manager; 0: Use ViewLocationsRenderedLastFrame"));
 
+static TAutoConsoleVariable<bool> CVarGrassVisualLogCameraLocations(
+	TEXT("grass.VisualLog.CameraLocations"),
+	false,
+	TEXT("Allows to visualize the cameras used for landscape grass in the visual logger "));
+
 static FAutoConsoleVariable CVarMaxAsyncNaniteProxiesPerSecond(
 	TEXT("landscape.Nanite.MaxAsyncProxyBuildsPerSecond"),
 	6.0f,
@@ -654,14 +660,25 @@ void ULandscapeSubsystem::RegenerateGrass(bool bInFlushGrass, bool bInForceSync,
 			}
 		}
 
+		UE_IFVLOG(
+			{
+				if (CVarGrassVisualLogCameraLocations.GetValueOnGameThread())
+				{
+					for (const FVector& CameraLocation : CameraLocations)
+					{
+						UE_VLOG_LOCATION(this, LogLandscape, Log, CameraLocation, 10.0f, FColor::Green, TEXT("Grass Camera Location"));
+					}
+				}
+			});
+
 		// Update the grass near the specified location(s) : 
 		for (TObjectPtr<ALandscapeProxy> ProxyPtr : Proxies)
-			{
+		{
 			ALandscapeProxy* Proxy = ProxyPtr.Get();
-				Proxy->UpdateGrass(CameraLocations, bInForceSync);
-			}
+			Proxy->UpdateGrass(CameraLocations, bInForceSync);
 		}
 	}
+}
 
 ETickableTickType ULandscapeSubsystem::GetTickableTickType() const
 {
