MoviePipeline: Fixed a regression that caused the pathtracer node in the graph to not generate output unless it's using tiled rendering.

#rnx

diff --git a/Engine/Plugins/MovieScene/MovieRenderPipeline/Source/MovieRenderPipelineRenderPasses/Private/Graph/Nodes/MovieGraphPathTracerPassNode.cpp b/Engine/Plugins/MovieScene/MovieRenderPipeline/Source/MovieRenderPipelineRenderPasses/Private/Graph/Nodes/MovieGraphPathTracerPassNode.cpp
--- a/Engine/Plugins/MovieScene/MovieRenderPipeline/Source/MovieRenderPipelineRenderPasses/Private/Graph/Nodes/MovieGraphPathTracerPassNode.cpp
+++ b/Engine/Plugins/MovieScene/MovieRenderPipeline/Source/MovieRenderPipelineRenderPasses/Private/Graph/Nodes/MovieGraphPathTracerPassNode.cpp
@@ -24,6 +24,7 @@ UMovieGraphPathTracerRenderPassNode::UMovieGraphPathTracerRenderPassNode()
 	, bWriteAllSamples(false)
 	, bDisableToneCurve(false)
 	, bAllowOCIO(true)
+	, TileCount(1)
 {
 	ShowFlags->ApplyDefaultShowFlagValue(VMI_PathTracing, true);
 	// TODO: Showflag for SetMotionBlur()?
@@ -244,6 +245,31 @@ FEngineShowFlags UMovieGraphPathTracerRenderPassNode::GetShowFlags() const
 	return OutShowFlags;
 }
 
+bool UMovieGraphPathTracerRenderPassNode::GetEnableHighResolutionTiling() const
+{
+	return bEnableHighResolutionTiling;
+}
+
+FIntPoint UMovieGraphPathTracerRenderPassNode::GetTileCount() const
+{
+	return bEnableHighResolutionTiling ? FIntPoint(TileCount, TileCount) : FIntPoint(1, 1);
+}
+
+float UMovieGraphPathTracerRenderPassNode::GetTileOverlapPercentage() const
+{
+	return bEnableHighResolutionTiling ? OverlapPercentage : 0;
+}
+
+bool UMovieGraphPathTracerRenderPassNode::GetEnablePageToSystemMemory() const
+{
+	return bPageToSystemMemory;
+}
+
+bool UMovieGraphPathTracerRenderPassNode::GetEnableHistoryPerTile() const
+{
+	return bAllocateHistoryPerTile;
+}
+
 void UMovieGraphPathTracerRenderPassNode::UpdateTelemetry(FMoviePipelineShotRenderTelemetry* InTelemetry) const
 {
 	InTelemetry->bUsesPathTracer = true;
diff --git a/Engine/Plugins/MovieScene/MovieRenderPipeline/Source/MovieRenderPipelineRenderPasses/Public/Graph/Nodes/MovieGraphPathTracerPassNode.h b/Engine/Plugins/MovieScene/MovieRenderPipeline/Source/MovieRenderPipelineRenderPasses/Public/Graph/Nodes/MovieGraphPathTracerPassNode.h
--- a/Engine/Plugins/MovieScene/MovieRenderPipeline/Source/MovieRenderPipelineRenderPasses/Public/Graph/Nodes/MovieGraphPathTracerPassNode.h
+++ b/Engine/Plugins/MovieScene/MovieRenderPipeline/Source/MovieRenderPipelineRenderPasses/Public/Graph/Nodes/MovieGraphPathTracerPassNode.h
@@ -56,11 +56,11 @@ public:
 	virtual bool GetAllowDenoiser() const override;
 	virtual TUniquePtr<UE::MovieGraph::Rendering::FMovieGraphImagePassBase> CreateInstance() const;
 	virtual FEngineShowFlags GetShowFlags() const override;
-	virtual bool GetEnableHighResolutionTiling() const override { return bEnableHighResolutionTiling; }
-	virtual FIntPoint GetTileCount() const override { return FIntPoint(TileCount, TileCount); }
-	virtual float GetTileOverlapPercentage() const override { return OverlapPercentage; }
-	virtual bool GetEnablePageToSystemMemory() const override { return bPageToSystemMemory; }
-	virtual bool GetEnableHistoryPerTile() const override { return bAllocateHistoryPerTile; }
+	virtual bool GetEnableHighResolutionTiling() const override;
+	virtual FIntPoint GetTileCount() const override;
+	virtual float GetTileOverlapPercentage() const override;
+	virtual bool GetEnablePageToSystemMemory() const override;
+	virtual bool GetEnableHistoryPerTile() const override;
 
 	// ~UMovieGraphImagePassBaseNode Interface
 
@@ -278,7 +278,7 @@ public:
 	/**
 	* Rendering effects such as Depth of Field may produce different results near the edge of a tile (as it cannot sample outside of the tile itself), so this
 	* setting allows you to create an overlapped region between tiles. 10% is a good starting point, but may need to be increased if you have extremely large
-	* depth of field.
+	* depth of field. Only has an effect if bEnableHighResolutionTiling is enabled.
 	*
 	* Note: This uses 0-50 and not 0-.5 like the previous system did to bring it in-line with other usages of overscan in the engine (nDisplay).
 	*/
