MoviePipeline: Fixed localization warning in MRG.

#rnx

diff --git a/Engine/Plugins/MovieScene/MovieRenderPipeline/Source/MovieRenderPipelineCore/Private/Graph/MovieGraphConfig.cpp b/Engine/Plugins/MovieScene/MovieRenderPipeline/Source/MovieRenderPipelineCore/Private/Graph/MovieGraphConfig.cpp
--- a/Engine/Plugins/MovieScene/MovieRenderPipeline/Source/MovieRenderPipelineCore/Private/Graph/MovieGraphConfig.cpp
+++ b/Engine/Plugins/MovieScene/MovieRenderPipeline/Source/MovieRenderPipelineCore/Private/Graph/MovieGraphConfig.cpp
@@ -2091,7 +2091,7 @@ UMovieGraphEvaluatedConfig* UMovieGraphConfig::CreateFlattenedGraph(const FMovie
 			{
 				if (!IsValid(Edge))
 				{
-					const FText InvalidNodeError = FText::Format(LOCTEXT("TraversalError_InvalidNode", "Found an invalid node in branch [{0}]. Aborting graph evaluation."), { FText::FromName(BranchName) });
+					const FText InvalidNodeError = FText::Format(LOCTEXT("TraversalError_InvalidNodeInBranch", "Found an invalid node in branch [{0}]. Aborting graph evaluation."), { FText::FromName(BranchName) });
 					UE_LOG(LogMovieRenderPipeline, Error, TEXT("%s"), *InvalidNodeError.ToString());
 					OutError = InvalidNodeError.ToString();
 					return nullptr;
