MoviePipeline: Fixed a rare crash that could occur with the Command Line Encoder node if it is used when there was no output data generated (eg, JPGs).

#rnx

diff --git a/Engine/Plugins/MovieScene/MovieRenderPipeline/Source/MovieRenderPipelineCore/Private/Graph/Nodes/MovieGraphCommandLineEncoderNode.cpp b/Engine/Plugins/MovieScene/MovieRenderPipeline/Source/MovieRenderPipelineCore/Private/Graph/Nodes/MovieGraphCommandLineEncoderNode.cpp
--- a/Engine/Plugins/MovieScene/MovieRenderPipeline/Source/MovieRenderPipelineCore/Private/Graph/Nodes/MovieGraphCommandLineEncoderNode.cpp
+++ b/Engine/Plugins/MovieScene/MovieRenderPipeline/Source/MovieRenderPipelineCore/Private/Graph/Nodes/MovieGraphCommandLineEncoderNode.cpp
@@ -93,6 +93,15 @@ void UMovieGraphCommandLineEncoderNode::StartEncodingProcess(TArray<FMovieGraphR
 		return;
 	}
 
+	// Early out if there isn't any data to feed to the encoder. Could happen if it's running on a branch with no nodes that generate output (eg, JPG).
+	if (InGeneratedData.IsEmpty())
+	{
+		UE_LOG(LogMovieRenderPipelineIO, Error, TEXT("The Command Line Encoder cannot run because no input files were generated for it to use."));
+		CachedPipeline->RequestShutdown(true);
+		
+		return;
+	}
+
 	// Get the base output file path; it still needs to be validated and resolved at this point
 	const UMovieGraphGlobalOutputSettingNode* OutputSettingNode = GetSettingOnBranch<UMovieGraphGlobalOutputSettingNode>();
 	FString FilePathFormatString = OutputSettingNode->OutputDirectory.Path / FileNameFormat;
