MoviePipeline: Fixed another crash that could occur if the graph is evaluated and contains a node from an unloaded plugin.

#rnx

diff --git a/Engine/Plugins/MovieScene/MovieRenderPipeline/Source/MovieRenderPipelineCore/Private/Graph/MovieGraphConfig.cpp b/Engine/Plugins/MovieScene/MovieRenderPipeline/Source/MovieRenderPipelineCore/Private/Graph/MovieGraphConfig.cpp
--- a/Engine/Plugins/MovieScene/MovieRenderPipeline/Source/MovieRenderPipelineCore/Private/Graph/MovieGraphConfig.cpp
+++ b/Engine/Plugins/MovieScene/MovieRenderPipeline/Source/MovieRenderPipelineCore/Private/Graph/MovieGraphConfig.cpp
@@ -1922,6 +1922,15 @@ bool UMovieGraphConfig::CreateFlattenedGraph_Recursive(UMovieGraphEvaluatedConfi
 	{
 		for (UMovieGraphEdge* Edge : Pin->Edges)
 		{
+			// If the edge is invalid, it's most likely due to a node from an unloaded plugin. Graph evaluation can't continue at this point.
+			// Ideally we could provide the branch name here, but it isn't available currently.
+			if (!IsValid(Edge))
+			{
+				InEvaluationContext.TraversalError = LOCTEXT("TraversalError_InvalidConnectionInBranch", "Found an invalid connection, likely from an invalid node. Aborting graph evaluation.");
+
+				return false;
+			}
+			
 			constexpr bool bFollowRerouteConnections = true;
 			UMovieGraphPin* OtherPin = Edge->GetOtherPin(Pin, bFollowRerouteConnections);
 			if (!OtherPin)
