MoviePipeline: Fixed a regression that caused a "check" to hit in MRG if connecting a pin to itself.

diff --git a/Engine/Plugins/MovieScene/MovieRenderPipeline/Source/MovieRenderPipelineEditor/Private/Graph/MovieGraphSchema.cpp b/Engine/Plugins/MovieScene/MovieRenderPipeline/Source/MovieRenderPipelineEditor/Private/Graph/MovieGraphSchema.cpp
--- a/Engine/Plugins/MovieScene/MovieRenderPipeline/Source/MovieRenderPipelineEditor/Private/Graph/MovieGraphSchema.cpp
+++ b/Engine/Plugins/MovieScene/MovieRenderPipeline/Source/MovieRenderPipelineEditor/Private/Graph/MovieGraphSchema.cpp
@@ -486,7 +486,6 @@ bool UMovieGraphSchema::TryCreateConnection(UEdGraphPin* InA, UEdGraphPin* InB)
 	check(InA && InB);
 	const UEdGraphPin* A = (InA->Direction == EGPD_Output) ? InA : InB;
 	const UEdGraphPin* B = (InA->Direction == EGPD_Input) ? InA : InB;
-	check(A->Direction == EGPD_Output && B->Direction == EGPD_Input);
 
 	const UMoviePipelineEdGraphNodeBase* EdGraphNodeA = CastChecked<UMoviePipelineEdGraphNodeBase>(A->GetOwningNode());
 	const UMoviePipelineEdGraphNodeBase* EdGraphNodeB = CastChecked<UMoviePipelineEdGraphNodeBase>(B->GetOwningNode());
