MoviePipeline: Fixed a potential editor dead-lock if there was an issue with the parent hierarchy in Sequencer content, resulting in the application freezing or users having to manually end the process.

diff --git a/Engine/Plugins/MovieScene/MovieRenderPipeline/Source/MovieRenderPipelineCore/Private/MoviePipelineBlueprintLibrary.cpp b/Engine/Plugins/MovieScene/MovieRenderPipeline/Source/MovieRenderPipelineCore/Private/MoviePipelineBlueprintLibrary.cpp
--- a/Engine/Plugins/MovieScene/MovieRenderPipeline/Source/MovieRenderPipelineCore/Private/MoviePipelineBlueprintLibrary.cpp
+++ b/Engine/Plugins/MovieScene/MovieRenderPipeline/Source/MovieRenderPipelineCore/Private/MoviePipelineBlueprintLibrary.cpp
@@ -597,6 +597,12 @@ void UMoviePipelineBlueprintLibrary::UpdateJobShotListFromSequence(ULevelSequenc
 									}
 								}
 							}
+							else
+							{
+								// If the content is set up wrong and doesn't have a valid parent, we need to break out of the loop
+								UE_LOG(LogMovieRenderPipeline, Error, TEXT("A camera binding was found while calculating multi-camera support, but the binding had an invalid parent (ParentGuid %s), skipping..."), *ParentGuid.ToString());
+								break;
+							}
 						}
 
 						if (bAddCamera)
