MoviePipeline: Updated Quick Render's "Current Viewport" mode to respect more camera properties when the viewport is piloted by cine cams.

#rnx

diff --git a/Engine/Plugins/MovieScene/MovieRenderPipeline/Source/MovieRenderPipelineEditor/Private/Graph/MovieGraphQuickRender.cpp b/Engine/Plugins/MovieScene/MovieRenderPipeline/Source/MovieRenderPipelineEditor/Private/Graph/MovieGraphQuickRender.cpp
--- a/Engine/Plugins/MovieScene/MovieRenderPipeline/Source/MovieRenderPipelineEditor/Private/Graph/MovieGraphQuickRender.cpp
+++ b/Engine/Plugins/MovieScene/MovieRenderPipeline/Source/MovieRenderPipelineEditor/Private/Graph/MovieGraphQuickRender.cpp
@@ -1079,11 +1079,27 @@ void UMovieGraphQuickRenderSubsystem::PerformPreRenderSetup(UWorld* InEditorWorl
 		// after the render finishes, which isn't ideal (and this will also create entries in the undo stack, also not ideal).
 		UWorld* PieWorld = GetWorldOfType(EWorldType::PIE);
 		
-		// Add a new camera that mimics the viewport's camera to the level sequence. Use a regular camera actor here rather than a cine camera. Using
-		// a cine camera massively complicates things and creates a situation where it's difficult to match what the viewport looks like (eg, exposure).
 		FActorSpawnParameters CameraSpawnParams;
 		CameraSpawnParams.ObjectFlags |= RF_Transient;
 		CameraSpawnParams.SpawnCollisionHandlingOverride = ESpawnActorCollisionHandlingMethod::AlwaysSpawn;
+
+		// If the viewport is piloting a camera (and not another type of non-camera actor), we should just duplicate it. We can't just copy over
+		// the basic transform + FOV, etc because this may be a cine camera with many properties set on it (like auto-exposure).
+		// Note that this supports "cinematic locks" (like cameras piloted via Sequencer) and normal piloted cameras.
+		if (CachedPrePieData.ViewportActorLock.IsValid() && IsValid(CachedPrePieData.ViewportActorLockCameraComponent))
+		{
+			AActor* ViewportPilotActor = CachedPrePieData.ViewportActorLock.Get();
+			if (ViewportPilotActor->IsA<ACameraActor>())
+			{
+				// Use the current pilot cam as the template
+				CameraSpawnParams.Template = ViewportPilotActor;
+
+				return Cast<ACameraActor>(PieWorld->SpawnActor(ViewportPilotActor->GetClass(), nullptr, CameraSpawnParams));
+			}
+		}
+		
+		// Add a new camera that mimics the viewport's camera to the level sequence. Use a regular camera actor here rather than a cine camera. Using
+		// a cine camera massively complicates things and creates a situation where it's difficult to match what the viewport looks like (eg, exposure).
 		ACameraActor* NewCamera = Cast<ACameraActor>(PieWorld->SpawnActor<ACameraActor>(CameraSpawnParams));
 
 		// Update the camera to look like the viewport's camera
@@ -1097,28 +1113,13 @@ void UMovieGraphQuickRenderSubsystem::PerformPreRenderSetup(UWorld* InEditorWorl
 			float CameraAspectRatio = ViewportClient->AspectRatio;
 			float CameraOrthoWidth = ViewportClient->Viewport->GetSizeXY().X * ViewportClient->GetOrthoUnitsPerPixel(ViewportClient->Viewport);
 			
-			// If the viewport is being piloted by a camera, then we should use the properties of that to render from instead of the viewport's
-			// non-piloted camera properties.
+			// If the viewport is being piloted, then we should use the properties of that to render from instead of the viewport's
+			// non-piloted camera properties. Note that this only covers non-camera pilot actors; camera pilot actors are covered above.
 			if (CachedPrePieData.ViewportActorLock.IsValid())
 			{
-				// Support basic position/rotation for non-camera actors
 				const TStrongObjectPtr<AActor> LockedActorPin = CachedPrePieData.ViewportActorLock.Pin();
 				CameraPosition = LockedActorPin->GetActorLocation();
 				CameraRotation = LockedActorPin->GetActorRotation();
-
-				// If the piloted actor is a camera, get all of its camera properties. Note that this supports "cinematic locks" (like cameras piloted
-				// via Sequencer) and normal piloted cameras.
-				if (IsValid(CachedPrePieData.ViewportActorLockCameraComponent))
-				{
-					const UCameraComponent* PilotCameraComponent = CachedPrePieData.ViewportActorLockCameraComponent;
-					
-					CameraPosition = PilotCameraComponent->GetTypedOuter<AActor>()->GetActorLocation();
-					CameraRotation = PilotCameraComponent->GetTypedOuter<AActor>()->GetActorRotation();
-					bIsPerspective = PilotCameraComponent->ProjectionMode == ECameraProjectionMode::Perspective;
-					CameraFov = PilotCameraComponent->FieldOfView;
-					CameraAspectRatio = PilotCameraComponent->AspectRatio;
-					CameraOrthoWidth = PilotCameraComponent->OrthoWidth;
-				}
 			}
 			else
 			{
