MoviePipeline: Fixed another case where an Unknown node can cause a crash when performing a disconnection.

#rnx

diff --git a/Engine/Plugins/MovieScene/MovieRenderPipeline/Source/MovieRenderPipelineEditor/Private/Graph/MovieGraphSchema.cpp b/Engine/Plugins/MovieScene/MovieRenderPipeline/Source/MovieRenderPipelineEditor/Private/Graph/MovieGraphSchema.cpp
--- a/Engine/Plugins/MovieScene/MovieRenderPipeline/Source/MovieRenderPipelineEditor/Private/Graph/MovieGraphSchema.cpp
+++ b/Engine/Plugins/MovieScene/MovieRenderPipeline/Source/MovieRenderPipelineEditor/Private/Graph/MovieGraphSchema.cpp
@@ -562,7 +562,12 @@ void UMovieGraphSchema::BreakSinglePinLink(UEdGraphPin* SourcePin, UEdGraphPin*
 
 	UMovieGraphNode* SourceRuntimeNode = SourcePipelineGraphNode->GetRuntimeNode();
 	UMovieGraphNode* TargetRuntimeNode = TargetPipelineGraphNode->GetRuntimeNode();
-	check(SourceRuntimeNode && TargetRuntimeNode);
+
+	if (!IsValid(SourceRuntimeNode) || !IsValid(TargetRuntimeNode))
+	{
+		UE_LOG(LogMovieRenderPipeline, Warning, TEXT("Cannot remove connection from node which is not currently valid (it may be from a plugin that's not currently loaded)."));
+		return;
+	}
 
 	UMovieGraphConfig* RuntimeGraph = SourceRuntimeNode->GetGraph();
 	check(RuntimeGraph);
