MovieGraph: Prevent warning about relative frame numbers not being used when the user was using relative frame numbers in the filename tokens. This was caused by a separate filepath generation happening at the same time that was intentionally replacing the frame number tokens with a placeholder for further usage.

diff --git a/Engine/Plugins/MovieScene/MovieRenderPipeline/Source/MovieRenderPipelineRenderPasses/Private/MovieGraphImageSequenceOutputNode.cpp b/Engine/Plugins/MovieScene/MovieRenderPipeline/Source/MovieRenderPipelineRenderPasses/Private/MovieGraphImageSequenceOutputNode.cpp
--- a/Engine/Plugins/MovieScene/MovieRenderPipeline/Source/MovieRenderPipelineRenderPasses/Private/MovieGraphImageSequenceOutputNode.cpp
+++ b/Engine/Plugins/MovieScene/MovieRenderPipeline/Source/MovieRenderPipelineRenderPasses/Private/MovieGraphImageSequenceOutputNode.cpp
@@ -116,8 +116,14 @@ FString UMovieGraphImageSequenceOutputNode::GetFrameTemplatedFileName(
 	{
 		FrameTemplatedFormatString = InFileNameFormatString.Replace(TEXT("{frame_number_shot_rel}"), *FramePlaceholder);
 	}
+
+	// If time dilation is being used, the parameters will ask ResolveFilenameFormatArguments to warn the user if there's no _rel frame number
+	// found, but we're intentionally overriding them above to be able to replace them with a completely unrelated token, so we don't actually
+	// want that warning to be produced.
+	FMovieGraphFilenameResolveParams ParamsCopy = InParams;
+	ParamsCopy.bForceRelativeFrameNumbers = false;
 	
-	return UMovieGraphBlueprintLibrary::ResolveFilenameFormatArguments(FrameTemplatedFormatString, InParams, OutMergedFormatArgs);
+	return UMovieGraphBlueprintLibrary::ResolveFilenameFormatArguments(FrameTemplatedFormatString, ParamsCopy, OutMergedFormatArgs);
 }
 
 void UMovieGraphImageSequenceOutputNode::OnReceiveImageDataImpl(UMovieGraphPipeline* InPipeline, UE::MovieGraph::FMovieGraphOutputMergerFrame* InRawFrameData, const TSet<FMovieGraphRenderDataIdentifier>& InMask)
