Add boolean for MoviePipelineFunctionalTestBase to toggle screenshot comparisons

diff --git a/Engine/Plugins/MovieScene/MovieRenderPipeline/Source/MovieRenderPipelineEditor/Private/MoviePipelineFunctionalTestBase.cpp b/Engine/Plugins/MovieScene/MovieRenderPipeline/Source/MovieRenderPipelineEditor/Private/MoviePipelineFunctionalTestBase.cpp
--- a/Engine/Plugins/MovieScene/MovieRenderPipeline/Source/MovieRenderPipelineEditor/Private/MoviePipelineFunctionalTestBase.cpp
+++ b/Engine/Plugins/MovieScene/MovieRenderPipeline/Source/MovieRenderPipelineEditor/Private/MoviePipelineFunctionalTestBase.cpp
@@ -92,10 +92,7 @@ void AMoviePipelineFunctionalTestBase::OnMoviePipelineFinished(FMoviePipelineOut
 	}
 	else
 	{
-		if (bPerformDiff)
-		{
-			CompareRenderOutputToGroundTruth(InOutputData);
-		}
+		CompareRenderOutputToGroundTruth(InOutputData);
 	}
 }
 
@@ -212,16 +209,16 @@ void AMoviePipelineFunctionalTestBase::CompareRenderOutputToGroundTruth(FMoviePi
 		FinishTest(EFunctionalTestResult::Failed, FString::Printf(TEXT("Mismatched number of shots between GroudTruth and New. Expected %d got %d."), NumShotsGT, NumShotsNew));
 		return;
 	}
-
+		
 	TMap<FString, FString> OldToNewImagesToCompare;
-
+	
 	for (int32 ShotIndex = 0; ShotIndex < GroundTruthData.ShotData.Num(); ShotIndex++)
 	{
 		// We compare them using indexes because we know that they have the same number of indexes,
 		// and the order should be deterministic. 
 		const FMoviePipelineShotOutputData & GroundTruthShot = GroundTruthData.ShotData[ShotIndex];
 		const FMoviePipelineShotOutputData & NewData = InOutputData.ShotData[ShotIndex];
-
+		
 		// And then repeat this for each render pass in the shot
 		int32 NumRenderPassesNew = NewData.RenderPassData.Num();
 		int32 NumRenderPassesGT = GroundTruthShot.RenderPassData.Num();
@@ -230,7 +227,7 @@ void AMoviePipelineFunctionalTestBase::CompareRenderOutputToGroundTruth(FMoviePi
 			FinishTest(EFunctionalTestResult::Failed, FString::Printf(TEXT("Mismatched number of shots between GroudTruth and New. Expected %d got %d."), NumRenderPassesGT, NumRenderPassesNew));
 			return;
 		}
-
+		
 		for (const TPair<FMoviePipelinePassIdentifier, FMoviePipelineRenderPassOutputData>& GroundTruthPair : GroundTruthShot.RenderPassData)
 		{
 			const FMoviePipelineRenderPassOutputData* NewDataOutputData = NewData.RenderPassData.Find(GroundTruthPair.Key);
@@ -239,7 +236,7 @@ void AMoviePipelineFunctionalTestBase::CompareRenderOutputToGroundTruth(FMoviePi
 				FinishTest(EFunctionalTestResult::Failed, FString::Printf(TEXT("Did not render pass '%s' from camera '%s' that is in the Ground Truth."), *GroundTruthPair.Key.Name, *GroundTruthPair.Key.CameraName));
 				return;
 			}
-
+			
 			// Now we can check that they output the same number of files
 			int32 NumOutputFilesNew = NewDataOutputData->FilePaths.Num();
 			int32 NumOutputFilesGT = GroundTruthPair.Value.FilePaths.Num();
@@ -248,7 +245,7 @@ void AMoviePipelineFunctionalTestBase::CompareRenderOutputToGroundTruth(FMoviePi
 				FinishTest(EFunctionalTestResult::Failed, FString::Printf(TEXT("Mismatched number of shots between GroudTruth and New. Expected %d got %d."), NumOutputFilesGT, NumOutputFilesNew));
 				return;
 			}
-
+			
 			// Now we'll just loop through them in lockstep (as MRQ should put them in order anyways) and just add
 			// them to an array to be tested later. We don't test that the sub-file names are exact matches right now.
 			for (int32 Index = 0; Index < NumOutputFilesGT; Index++)
@@ -259,14 +256,17 @@ void AMoviePipelineFunctionalTestBase::CompareRenderOutputToGroundTruth(FMoviePi
 			}
 		}
 	}
-
-	// Time for the computationally expensive part, doing image comparisons!
-	TSharedPtr<FImageComparisonResult> ComparisonResult = ScreenshotManager->CompareImageSequence(OldToNewImagesToCompare, MetaData);
-	if (ComparisonResult.IsValid() && !ComparisonResult->AreSimilar())
+		
+	if (bPerformDiff)
 	{
-		ScreenshotManager->NotifyAutomationTestFrameworkOfImageComparison(*ComparisonResult);
-		FinishTest(EFunctionalTestResult::Failed, TEXT("Frames failed comparison tolerance!"));
-		return;
+		// Time for the computationally expensive part, doing image comparisons!
+		TSharedPtr<FImageComparisonResult> ComparisonResult = ScreenshotManager->CompareImageSequence(OldToNewImagesToCompare, MetaData);
+		if (ComparisonResult.IsValid() && !ComparisonResult->AreSimilar())
+		{
+			ScreenshotManager->NotifyAutomationTestFrameworkOfImageComparison(*ComparisonResult);
+			FinishTest(EFunctionalTestResult::Failed, TEXT("Frames failed comparison tolerance!"));
+			return;
+		}
 	}
 
 	UE_LOG(LogTemp, Display, TEXT("All image sequences from %s are similar to the Ground Truth."), *MetaData.ScreenShotName);
diff --git a/Engine/Plugins/MovieScene/MovieRenderPipeline/Source/MovieRenderPipelineEditor/Private/MoviePipelineFunctionalTestBase.h b/Engine/Plugins/MovieScene/MovieRenderPipeline/Source/MovieRenderPipelineEditor/Private/MoviePipelineFunctionalTestBase.h
--- a/Engine/Plugins/MovieScene/MovieRenderPipeline/Source/MovieRenderPipelineEditor/Private/MoviePipelineFunctionalTestBase.h
+++ b/Engine/Plugins/MovieScene/MovieRenderPipeline/Source/MovieRenderPipelineEditor/Private/MoviePipelineFunctionalTestBase.h
@@ -51,7 +51,12 @@ protected:
 
 	UPROPERTY(EditAnywhere, Category = "Movie Pipeline")
 	FImageTolerance CustomImageTolerance;
-	
+
+	/**
+	* Boolean flag to toggle screenshot pixel comparison of rendered versus ground truth images.
+	* True - screenshot pixel comparison will be performed.
+	* False - screenshot pixel comparison will not be performed.  Expects ground truth images to exist to verify that the required number of images were rendered, even if those ground truth images are not used in diffs.
+	*/	
 	UPROPERTY(EditAnywhere, Category = "Movie Pipeline")
 	bool bPerformDiff;
 
