MRG: Fixed alpha output (r.PostProcessing.PropagateAlpha) previous state restoration after renders.

Incorrect behavior introduced [at] CL 41486523, where the difference between MRQ & MRG was accidentally omitted.

diff --git a/Engine/Plugins/MovieScene/MovieRenderPipeline/Source/MovieRenderPipelineCore/Private/Graph/Nodes/MovieGraphGlobalGameOverrides.cpp b/Engine/Plugins/MovieScene/MovieRenderPipeline/Source/MovieRenderPipelineCore/Private/Graph/Nodes/MovieGraphGlobalGameOverrides.cpp
--- a/Engine/Plugins/MovieScene/MovieRenderPipeline/Source/MovieRenderPipelineCore/Private/Graph/Nodes/MovieGraphGlobalGameOverrides.cpp
+++ b/Engine/Plugins/MovieScene/MovieRenderPipeline/Source/MovieRenderPipelineCore/Private/Graph/Nodes/MovieGraphGlobalGameOverrides.cpp
@@ -190,7 +190,7 @@ void UMovieGraphGlobalGameOverridesNode::ApplySettings(const bool bOverrideValue
 		MOVIEPIPELINE_STORE_AND_OVERRIDE_CVAR_INT(NodeCDO->PreviousLumenSurfaceCacheRecaptureEveryFrame, TEXT("r.LumenScene.SurfaceCache.RecaptureEveryFrame"), 1, bOverrideValues);
 	}
 
-	MOVIEPIPELINE_STORE_AND_OVERRIDE_CVAR_BOOL(PreviousAlphaOutput, TEXT("r.PostProcessing.PropagateAlpha"), UE::MoviePipeline::GetAlphaOutputOverride(), bOverrideValues);
+	MOVIEPIPELINE_STORE_AND_OVERRIDE_CVAR_BOOL(NodeCDO->PreviousAlphaOutput, TEXT("r.PostProcessing.PropagateAlpha"), UE::MoviePipeline::GetAlphaOutputOverride(), bOverrideValues);
 
 	// Set cvars that allow MRQ to run correctly. These are not exposed as properties on the node, but the chance of a
 	// user needing to customize these (or be aware of them) is very low.
