MoviePipeline: Fixed an issue where the Selected Cameras mode didn't work in Quick Render if more than one camera was selected (only the first camera was rendered).

diff --git a/Engine/Plugins/MovieScene/MovieRenderPipeline/Source/MovieRenderPipelineCore/Public/Graph/MovieGraphQuickRenderSettings.h b/Engine/Plugins/MovieScene/MovieRenderPipeline/Source/MovieRenderPipelineCore/Public/Graph/MovieGraphQuickRenderSettings.h
--- a/Engine/Plugins/MovieScene/MovieRenderPipeline/Source/MovieRenderPipelineCore/Public/Graph/MovieGraphQuickRenderSettings.h
+++ b/Engine/Plugins/MovieScene/MovieRenderPipeline/Source/MovieRenderPipelineCore/Public/Graph/MovieGraphQuickRenderSettings.h
@@ -179,7 +179,7 @@ public:
 	UPROPERTY(EditAnywhere, BlueprintReadWrite, Category="Quick Render", DisplayName="Apply Viewport Look", meta=(EditCondition="bOverride_ViewportLookFlags", Bitmask, BitmaskEnum="/Script/MovieRenderPipelineCore.EMovieGraphQuickRenderViewportLookFlags"))
 	int32 ViewportLookFlags = static_cast<int32>(EMovieGraphQuickRenderViewportLookFlags::None);
 
-	/** The frame range that should be used to render from. Only available in modes that render more than one frame. */
+	/** The frame range that should be used to render from. Not applicable to modes that only render stills. */
 	UPROPERTY(EditAnywhere, BlueprintReadWrite, Category="Quick Render", DisplayName="Sequencer Frame Range")
 	EMovieGraphQuickRenderFrameRangeType FrameRangeType = EMovieGraphQuickRenderFrameRangeType::PlaybackRange;
 
diff --git a/Engine/Plugins/MovieScene/MovieRenderPipeline/Source/MovieRenderPipelineEditor/Private/Graph/MovieGraphQuickRender.cpp b/Engine/Plugins/MovieScene/MovieRenderPipeline/Source/MovieRenderPipelineEditor/Private/Graph/MovieGraphQuickRender.cpp
--- a/Engine/Plugins/MovieScene/MovieRenderPipeline/Source/MovieRenderPipelineEditor/Private/Graph/MovieGraphQuickRender.cpp
+++ b/Engine/Plugins/MovieScene/MovieRenderPipeline/Source/MovieRenderPipelineEditor/Private/Graph/MovieGraphQuickRender.cpp
@@ -45,7 +45,9 @@
 #include "Editor/Transactor.h"
 #include "EditorViewportClient.h"
 #include "LevelEditorViewport.h"
+#include "PackageHelperFunctions.h"
 #include "Selection.h"
+#include "Subsystems/AssetEditorSubsystem.h"
 #endif
 
 #if WITH_OCIO
@@ -55,6 +57,49 @@
 
 #define LOCTEXT_NAMESPACE "FMovieGraphQuickRender"
 
+namespace
+{
+	/**
+	 * Writes a target level sequence to disk and opens it in Sequencer. Generally this is useful for debugging transient level sequences generated
+	 * by Quick Render. InAssetPath should be of the format /Game/Temp/<AssetName>.
+	 */
+	void DEBUG_SaveAndOpenSequence(const ULevelSequence* InTargetLevelSequence, const FString& InAssetPath)
+	{
+		if (!IsValid(InTargetLevelSequence))
+		{
+			UE_LOG(LogMovieRenderPipeline, Error, TEXT("[%hs]: Need a valid, non-null level sequence to save/open."), __FUNCTION__);
+			return;
+		}
+
+		const FString PackageFilename = FPackageName::LongPackageNameToFilename(*InAssetPath, FPackageName::GetAssetPackageExtension());
+
+		// Make the new package which will contain the (duplicated) level sequence.
+		UPackage* NewPackage = CreatePackage(*InAssetPath);
+		if (!NewPackage)
+		{
+			UE_LOG(LogMovieRenderPipeline, Error, TEXT("[%hs]: Could not create new package."), __FUNCTION__);
+			return;
+		}
+
+		// Duplicate the target level sequence into the new package.
+		ULevelSequence* NewLevelSequence = CastChecked<ULevelSequence>(StaticDuplicateObject(InTargetLevelSequence, NewPackage));
+		if (!NewLevelSequence)
+		{
+			UE_LOG(LogMovieRenderPipeline, Error, TEXT("[%hs]: Unable to duplicate target sequence."), __FUNCTION__);
+			return;
+		}
+
+		// To ensure that SavePackageHelper() writes to disk, the new level sequence cannot have RF_Transient
+		NewLevelSequence->ClearFlags(RF_Transient);
+		NewLevelSequence->SetFlags(RF_Standalone | RF_Public);
+		NewLevelSequence->MarkPackageDirty();
+
+		// Save the level sequence to disk and open in Sequencer.
+		SavePackageHelper(NewPackage, *PackageFilename, RF_Standalone, GWarn, SAVE_None);
+		GEditor->GetEditorSubsystem<UAssetEditorSubsystem>()->OpenEditorForAsset(NewLevelSequence);
+	}
+}
+
 /**
  * Wraps the scope that this is created in with a dummy transaction. When going out of scope, the dummy transaction will be canceled, effectively
  * blocking any transactions that occurred while in scope.
@@ -140,13 +185,13 @@ void UMovieGraphQuickRenderSubsystem::BeginQuickRender(const EMovieGraphQuickRen
 		return;
 	}
 
+	// Cache out anything necessary before PIE starts
+	CachePrePieData(EditorWorld);
+
 	// Update the graph as needed in order to properly apply the settings specified by the user. This will perform a graph evaluation, so the job
 	// should be fully initialized at this point.
 	ApplyQuickRenderUpdatesToDuplicateGraph(EditorWorld);
 
-	// Cache out anything necessary before PIE starts
-	CachePrePieData(EditorWorld);
-
 	// Create the executor. Note that it's allowed to use unsaved levels (many animators work off of unsaved levels). Because Quick Render is
 	// creating a temporary queue and rendering locally, using an unsaved level is safe.
 	// Always use the PIE executor for Quick Render, rather than the executor specified in the project settings for normal MRG.
@@ -157,6 +202,9 @@ void UMovieGraphQuickRenderSubsystem::BeginQuickRender(const EMovieGraphQuickRen
 	TemporaryExecutor->OnIndividualJobStarted().AddLambda([this, EditorWorld](UMoviePipelineExecutorJob* InJob)
 	{
 		PerformPreRenderSetup(EditorWorld);
+
+		// DEBUG: If the rendering level sequence needs to be inspected in Sequencer for debugging, uncomment this line.
+		// DEBUG_SaveAndOpenSequence(RenderingLevelSequence.Get(), TEXT("/Game/Temp/TempRenderingLevelSequence"));
 	});
 
 	// View the job's rendered frame(s)/video(s) if requested
@@ -347,8 +395,13 @@ void UMovieGraphQuickRenderSubsystem::ApplyQuickRenderUpdatesToDuplicateGraph(co
 		const TRange<FFrameNumber> PlaybackRange = GetPlaybackRange();
 		if (!PlaybackRange.IsEmpty())
 		{
-			SetStartFrame(PlaybackRange.GetLowerBoundValue().Value);
-			SetEndFrame(PlaybackRange.GetUpperBoundValue().Value);
+			const int32 StartFrame = PlaybackRange.GetLowerBoundValue().Value;
+			const int32 EndFrame = PlaybackRange.GetUpperBoundValue().Value;
+			
+			SetStartFrame(StartFrame);
+			SetEndFrame(EndFrame);
+			
+			UE_LOG(LogMovieRenderPipeline, Display, TEXT("Quick Render Start Frame: %d, End Frame: %d"), StartFrame, EndFrame);
 		}
 	}
 
@@ -940,12 +993,18 @@ TRange<FFrameNumber> UMovieGraphQuickRenderSubsystem::GetPlaybackRange()
 	FFrameNumber StartFrame;
 	FFrameNumber EndFrame;
 
-	// If we're in Current Viewport or Selected Cameras mode, the playback range is always one frame long (only a single frame is rendered).
-	if ((QuickRenderMode == EMovieGraphQuickRenderMode::CurrentViewport) || (QuickRenderMode == EMovieGraphQuickRenderMode::SelectedCameras))
+	// If we're in Current Viewport, the playback range is always one frame long (only a single frame is rendered).
+	if (QuickRenderMode == EMovieGraphQuickRenderMode::CurrentViewport)
 	{
 		return TRange<FFrameNumber>::Inclusive(0, 1);
 	}
 
+	// The Selected Cameras mode renders one frame for each camera.
+	if (QuickRenderMode == EMovieGraphQuickRenderMode::SelectedCameras)
+	{
+		return TRange<FFrameNumber>::Inclusive(0, CachedPrePieData.SelectedCameras.Num());
+	}
+
 	if (QuickRenderModeSettings->FrameRangeType == EMovieGraphQuickRenderFrameRangeType::SelectionRange)
 	{
 		const int32 SelectionStartFrame = ULevelSequenceEditorBlueprintLibrary::GetSelectionRangeStart();
@@ -1045,7 +1104,7 @@ void UMovieGraphQuickRenderSubsystem::PerformPreRenderSetup(UWorld* InEditorWorl
 		const int32 NumCameraCutSections = CameraCutTrack->GetAllSections().Num();
 
 		// Determine the playback range for the new camera cut. Each camera cut section has a playback range of one frame.
-		const FFrameNumber StartFrame = MovieScene->GetPlaybackRange().GetLowerBound().GetValue();
+		const FFrameNumber StartFrame = MovieScene->GetPlaybackRange().GetUpperBound().GetValue();
 		const FFrameNumber EndFrame = FFrameRate::TransformTime(FFrameTime(NumCameraCutSections + 1), MovieScene->GetDisplayRate(), MovieScene->GetTickResolution()).GetFrame();
 
 		// The entire playback range for the movie scene always starts a 0 and ends at the last camera cut section's end frame, unless there's an override.
