MoviePipeline: Addressed a couple of play-after-render issues in Quick Render.
* Fixed an issue that caused multi-layer EXRs to not open in external viewers.
* Fixed an issue that caused play-after-render to fail if there are spaces in the file path(s) that are provided to the external viewer.

diff --git a/Engine/Plugins/MovieScene/MovieRenderPipeline/Source/MovieRenderPipelineEditor/Private/MoviePipelinePostRenderFileDisplayProcessor.cpp b/Engine/Plugins/MovieScene/MovieRenderPipeline/Source/MovieRenderPipelineEditor/Private/MoviePipelinePostRenderFileDisplayProcessor.cpp
--- a/Engine/Plugins/MovieScene/MovieRenderPipeline/Source/MovieRenderPipelineEditor/Private/MoviePipelinePostRenderFileDisplayProcessor.cpp
+++ b/Engine/Plugins/MovieScene/MovieRenderPipeline/Source/MovieRenderPipelineEditor/Private/MoviePipelinePostRenderFileDisplayProcessor.cpp
@@ -460,7 +460,7 @@ void FMoviePipelinePostRenderFileDisplayProcessor::LaunchFilesWithSettings(const
 	const FString StartFrame = InFrameRangeToOpen.Key;
 	const FString EndFrame = InFrameRangeToOpen.Value;
 	
-	auto ReplaceFramePlaceholderInFiles = [&StartFrame, &InFilesToOpen]()
+	auto ReplaceFramePlaceholderInFilePaths = [&StartFrame](TArray<FString>& InFilePaths)
 	{
 		FString FramePlaceholder;
 		
@@ -482,8 +482,7 @@ void FMoviePipelinePostRenderFileDisplayProcessor::LaunchFilesWithSettings(const
 		}
 		
 		// Make sure the paths are in the right format for the platform, and have any frame templating applied
-		TArray<FString> TransformedFilePaths = InFilesToOpen;
-		for (FString& FileToOpen : TransformedFilePaths)
+		for (FString& FileToOpen : InFilePaths)
 		{
 			FPaths::MakePlatformFilename(FileToOpen);
 
@@ -492,8 +491,6 @@ void FMoviePipelinePostRenderFileDisplayProcessor::LaunchFilesWithSettings(const
 				FileToOpen.ReplaceInline(TEXT("{frame_placeholder}"), *FramePlaceholder);
 			}
 		}
-
-		return TransformedFilePaths;
 	};
 
 	auto InjectFrameRange = [this, &StartFrame, &EndFrame](FString& InConcatenatedFilePaths)
@@ -515,6 +512,14 @@ void FMoviePipelinePostRenderFileDisplayProcessor::LaunchFilesWithSettings(const
 			}
 		}
 	};
+
+	// Add quotes to the file paths (to avoid issues with spaces in the paths). This needs to be done for both the OS viewer and custom viewer.
+	TArray<FString> TransformedFilePaths;
+	TransformedFilePaths.Reserve(InFilesToOpen.Num());
+	for (const FString& Path : InFilesToOpen)
+	{
+		TransformedFilePaths.Add(FString::Format(TEXT("\"{0}\""), {Path}));
+	}
 	
 	if (InPlayOptions->PlaybackMethod == EMovieGraphPlaybackMethod::OperatingSystem)
 	{
@@ -529,7 +534,7 @@ void FMoviePipelinePostRenderFileDisplayProcessor::LaunchFilesWithSettings(const
 		FPaths::NormalizeFilename(ExecutablePathNoQuotes);
 		
 		// Replace {frame_placeholder}, if present, with the appropriate symbol (like '#')
-		const TArray<FString> TransformedFilePaths = ReplaceFramePlaceholderInFiles();
+		ReplaceFramePlaceholderInFilePaths(TransformedFilePaths);
 
 		// Depending on the frame range notation chosen, the command line arguments may need to be augmented with the frame range that was rendered
 		FString ConcatenatedFilePaths = FString::Join(TransformedFilePaths, TEXT(" "));
diff --git a/Engine/Plugins/MovieScene/MovieRenderPipeline/Source/MovieRenderPipelineRenderPasses/Private/MovieGraphImageSequenceOutputNode.cpp b/Engine/Plugins/MovieScene/MovieRenderPipeline/Source/MovieRenderPipelineRenderPasses/Private/MovieGraphImageSequenceOutputNode.cpp
--- a/Engine/Plugins/MovieScene/MovieRenderPipeline/Source/MovieRenderPipelineRenderPasses/Private/MovieGraphImageSequenceOutputNode.cpp
+++ b/Engine/Plugins/MovieScene/MovieRenderPipeline/Source/MovieRenderPipelineRenderPasses/Private/MovieGraphImageSequenceOutputNode.cpp
@@ -744,7 +744,7 @@ FString UMovieGraphImageSequenceOutputNode_MultiLayerEXR::ResolveOutputFilename(
 
 	FString FinalFilePath = UMovieGraphBlueprintLibrary::ResolveFilenameFormatArguments(FilePathFormatString, Params, OutResolveArgs);
 
-	OutFrameTemplatedFilename = GetFrameTemplatedFileName(Params, FileNameFormatString, OutResolveArgs);
+	OutFrameTemplatedFilename = GetFrameTemplatedFileName(Params, FilePathFormatString, OutResolveArgs);
 
 	if (FPaths::IsRelative(FinalFilePath))
 	{
