MRG: Minor update to previous state type.

Identical behavior but now uniform with macro & MRQ. Could not be previously included in the hotfix due as a header change.

diff --git a/Engine/Plugins/MovieScene/MovieRenderPipeline/Source/MovieRenderPipelineCore/Public/Graph/Nodes/MovieGraphGlobalGameOverrides.h b/Engine/Plugins/MovieScene/MovieRenderPipeline/Source/MovieRenderPipelineCore/Public/Graph/Nodes/MovieGraphGlobalGameOverrides.h
--- a/Engine/Plugins/MovieScene/MovieRenderPipeline/Source/MovieRenderPipelineCore/Public/Graph/Nodes/MovieGraphGlobalGameOverrides.h
+++ b/Engine/Plugins/MovieScene/MovieRenderPipeline/Source/MovieRenderPipelineCore/Public/Graph/Nodes/MovieGraphGlobalGameOverrides.h
@@ -245,7 +245,7 @@ private:
 	int32 PreviousChaosClothUseTimeStepSmoothing;
 	int32 PreviousSkipWaterInfoTextureRenderWhenWorldRenderingDisabled;
 	int32 PreviousNaniteVSMInvalidateOnLODDelta;
-	int32 PreviousAlphaOutput;
+	bool PreviousAlphaOutput;
 
 #if WITH_EDITOR
 	int32 PreviousGeoCacheStreamerShowNotification;
