MoviePipeline: Updated MRQ's queue asset definition so it opens the MRQ tab when a queue is opened from the content browser (previously, the queue was just loaded into the queue subsystem, which would only update an already-opened MRQ tab).

diff --git a/Engine/Plugins/MovieScene/MovieRenderPipeline/Source/MovieRenderPipelineEditor/Private/AssetDefinition_PipelineQueue.cpp b/Engine/Plugins/MovieScene/MovieRenderPipeline/Source/MovieRenderPipelineEditor/Private/AssetDefinition_PipelineQueue.cpp
--- a/Engine/Plugins/MovieScene/MovieRenderPipeline/Source/MovieRenderPipelineEditor/Private/AssetDefinition_PipelineQueue.cpp
+++ b/Engine/Plugins/MovieScene/MovieRenderPipeline/Source/MovieRenderPipelineEditor/Private/AssetDefinition_PipelineQueue.cpp
@@ -3,7 +3,9 @@
 #include "AssetDefinition_PipelineQueue.h"
 
 #include "Editor.h"
+#include "LevelEditor.h"
 #include "MoviePipelineQueueSubsystem.h"
+#include "MovieRenderPipelineEditorModule.h"
 
 EAssetCommandResult UAssetDefinition_PipelineQueue::OpenAssets(const FAssetOpenArgs& OpenArgs) const
 {
@@ -12,7 +14,14 @@ EAssetCommandResult UAssetDefinition_PipelineQueue::OpenAssets(const FAssetOpenA
 		UMoviePipelineQueueSubsystem* Subsystem = GEditor->GetEditorSubsystem<UMoviePipelineQueueSubsystem>();
 		check(Subsystem);
 
-		Subsystem->LoadQueue(ValidQueue);
+		if (Subsystem->LoadQueue(ValidQueue))
+		{
+			const FLevelEditorModule& LevelEditorModule = FModuleManager::GetModuleChecked<FLevelEditorModule>("LevelEditor");
+			if (const TSharedPtr<FTabManager> LevelEditorTabManager = LevelEditorModule.GetLevelEditorTabManager())
+			{
+				LevelEditorTabManager->TryInvokeTab(FMovieRenderPipelineEditorModule::MoviePipelineQueueTabName);
+			}
+		}
 	}
 
 	return EAssetCommandResult::Handled;
