MRG: Correctly restore r.Nanite.VSMInvalidateOnLODDelta console variable after render.

Identical to CL 43155545 hotfix, issue was discovered after submission.

diff --git a/Engine/Plugins/MovieScene/MovieRenderPipeline/Source/MovieRenderPipelineCore/Private/Graph/Nodes/MovieGraphGlobalGameOverrides.cpp b/Engine/Plugins/MovieScene/MovieRenderPipeline/Source/MovieRenderPipelineCore/Private/Graph/Nodes/MovieGraphGlobalGameOverrides.cpp
--- a/Engine/Plugins/MovieScene/MovieRenderPipeline/Source/MovieRenderPipelineCore/Private/Graph/Nodes/MovieGraphGlobalGameOverrides.cpp
+++ b/Engine/Plugins/MovieScene/MovieRenderPipeline/Source/MovieRenderPipelineCore/Private/Graph/Nodes/MovieGraphGlobalGameOverrides.cpp
@@ -233,7 +233,7 @@ void UMovieGraphGlobalGameOverridesNode::ApplySettings(const bool bOverrideValue
 		IConsoleVariable* VTInvalidateCvar = IConsoleManager::Get().FindConsoleVariable(TEXT("MoviePipeline.EnableVTInvalidateOnNaniteLOD"));
 		if (VTInvalidateCvar && VTInvalidateCvar->GetInt() != 0)
 		{
-			MOVIEPIPELINE_STORE_AND_OVERRIDE_CVAR_INT(PreviousNaniteVSMInvalidateOnLODDelta, TEXT("r.Nanite.VSMInvalidateOnLODDelta"), 1, bOverrideValues);
+			MOVIEPIPELINE_STORE_AND_OVERRIDE_CVAR_INT(NodeCDO->PreviousNaniteVSMInvalidateOnLODDelta, TEXT("r.Nanite.VSMInvalidateOnLODDelta"), 1, bOverrideValues);
 		}
 
 	}
