MoviePipeline: Fixed an issue where a comment node could not be deleted in MRG (regression from 42607060).

diff --git a/Engine/Plugins/MovieScene/MovieRenderPipeline/Source/MovieRenderPipelineEditor/Private/Widgets/Graph/SMovieGraphConfigPanel.cpp b/Engine/Plugins/MovieScene/MovieRenderPipeline/Source/MovieRenderPipelineEditor/Private/Widgets/Graph/SMovieGraphConfigPanel.cpp
--- a/Engine/Plugins/MovieScene/MovieRenderPipeline/Source/MovieRenderPipelineEditor/Private/Widgets/Graph/SMovieGraphConfigPanel.cpp
+++ b/Engine/Plugins/MovieScene/MovieRenderPipeline/Source/MovieRenderPipelineEditor/Private/Widgets/Graph/SMovieGraphConfigPanel.cpp
@@ -270,12 +270,23 @@ bool SMoviePipelineGraphPanel::CanDeleteSelectedNodes() const
 	for (UObject* Object : GraphEditorWidget->GetSelectedNodes())
 	{
 		const UMoviePipelineEdGraphNodeBase* GraphNode = Cast<UMoviePipelineEdGraphNodeBase>(Object);
-		if (!GraphNode || !GraphNode->CanUserDeleteNode())
+
+		// If the editor node doesn't derive from UMoviePipelineEdGraphNodeBase, it's a node from outside MRG (eg, a comment node); these are
+		// always deletable. They will not have a runtime representation. Ideally we could check to see if this is in CurrentGraph->GetEditorOnlyNodes(),
+		// but that array is only updated after a save is performed.
+		if (!GraphNode)
+		{
+			continue;
+		}
+
+		// Some nodes in the graph are explicitly marked as non-deletable (eg, inputs/outputs).
+		if (!GraphNode->CanUserDeleteNode())
 		{
 			return false;
 		}
 
 		// Additionally, if the associated runtime node is not valid (likely from an unloaded plugin), then the node cannot be deleted.
+		// Editor-only nodes do not have an associated runtime node, but those are accounted for earlier.
 		if (!GraphNode->GetRuntimeNode())
 		{
 			return false;
