MoviePipeline: Fix rounding discrepancies between the constrained view rect & letterboxing.

* FViewport::CalculateViewExtents is used by movie pipeline's render method to calculate constrained view rectangles.
* After rendering, letterbox black bars are drawn, only with a slightly different calculation causing misalignment due to rounding error differences (CeilToInt vs RoundToInt).

diff --git a/Engine/Plugins/MovieScene/MovieRenderPipeline/Source/MovieRenderPipelineRenderPasses/Private/Graph/Renderers/MovieGraphDeferredPass.cpp b/Engine/Plugins/MovieScene/MovieRenderPipeline/Source/MovieRenderPipelineRenderPasses/Private/Graph/Renderers/MovieGraphDeferredPass.cpp
--- a/Engine/Plugins/MovieScene/MovieRenderPipeline/Source/MovieRenderPipelineRenderPasses/Private/Graph/Renderers/MovieGraphDeferredPass.cpp
+++ b/Engine/Plugins/MovieScene/MovieRenderPipeline/Source/MovieRenderPipelineRenderPasses/Private/Graph/Renderers/MovieGraphDeferredPass.cpp
@@ -649,9 +649,10 @@ void FMovieGraphDeferredPass::PostRendererSubmission(
 		const float OutputSizeAspectRatio = FullOutputSize.X / (float)FullOutputSize.Y;
 		const float CameraAspectRatio = InCameraInfo.bAllowCameraAspectRatio ? CameraCache.AspectRatio : OutputSizeAspectRatio;
 	
+		// Note: We use RoundToInt to match the unscaled/constrained view rectangle calculated in FViewport::CalculateViewExtents
 		const FIntPoint ConstrainedFullSize = CameraAspectRatio > OutputSizeAspectRatio ?
-			FIntPoint(FullOutputSize.X, FMath::CeilToInt((double)FullOutputSize.X / (double)CameraAspectRatio)) :
-			FIntPoint(FMath::CeilToInt(CameraAspectRatio * FullOutputSize.Y), FullOutputSize.Y);
+			FIntPoint(FullOutputSize.X, FMath::RoundToInt((double)FullOutputSize.X / (double)CameraAspectRatio)) :
+			FIntPoint(FMath::RoundToInt(CameraAspectRatio * FullOutputSize.Y), FullOutputSize.Y);
 	
 		const FIntPoint TileViewMin = InSampleState.OverlappedOffset;
 		const FIntPoint TileViewMax = TileViewMin + InSampleState.BackbufferResolution;
diff --git a/Engine/Plugins/MovieScene/MovieRenderPipeline/Source/MovieRenderPipelineRenderPasses/Private/MoviePipelineDeferredPasses.cpp b/Engine/Plugins/MovieScene/MovieRenderPipeline/Source/MovieRenderPipelineRenderPasses/Private/MoviePipelineDeferredPasses.cpp
--- a/Engine/Plugins/MovieScene/MovieRenderPipeline/Source/MovieRenderPipelineRenderPasses/Private/MoviePipelineDeferredPasses.cpp
+++ b/Engine/Plugins/MovieScene/MovieRenderPipeline/Source/MovieRenderPipelineRenderPasses/Private/MoviePipelineDeferredPasses.cpp
@@ -1059,9 +1059,10 @@ void UMoviePipelineDeferredPassBase::PostRendererSubmission(const FMoviePipeline
 		const float OutputSizeAspectRatio = FullOutputSize.X / (float)FullOutputSize.Y;
 		const float CameraAspectRatio = bAllowCameraAspectRatio ? CameraCache.AspectRatio : OutputSizeAspectRatio;
 
+		// Note: We use RoundToInt to match the unscaled/constrained view rectangle calculated in FViewport::CalculateViewExtents
 		const FIntPoint ConstrainedFullSize = CameraAspectRatio > OutputSizeAspectRatio ?
-			FIntPoint(FullOutputSize.X, FMath::CeilToInt((double)FullOutputSize.X / (double)CameraAspectRatio)) :
-			FIntPoint(FMath::CeilToInt(CameraAspectRatio * FullOutputSize.Y), FullOutputSize.Y);
+			FIntPoint(FullOutputSize.X, FMath::RoundToInt((double)FullOutputSize.X / (double)CameraAspectRatio)) :
+			FIntPoint(FMath::RoundToInt(CameraAspectRatio * FullOutputSize.Y), FullOutputSize.Y);
 
 		const FIntPoint TileViewMin = InSampleState.OverlappedOffset;
 		const FIntPoint TileViewMax = TileViewMin + InSampleState.BackbufferSize;
diff --git a/Engine/Plugins/Runtime/nDisplay/Source/DisplayClusterMoviePipeline/Private/DisplayClusterMoviePipelineViewportPass.cpp b/Engine/Plugins/Runtime/nDisplay/Source/DisplayClusterMoviePipeline/Private/DisplayClusterMoviePipelineViewportPass.cpp
--- a/Engine/Plugins/Runtime/nDisplay/Source/DisplayClusterMoviePipeline/Private/DisplayClusterMoviePipelineViewportPass.cpp
+++ b/Engine/Plugins/Runtime/nDisplay/Source/DisplayClusterMoviePipeline/Private/DisplayClusterMoviePipelineViewportPass.cpp
@@ -403,8 +403,8 @@ void UDisplayClusterMoviePipelineViewportPassBase::GetViewportCutOffset(const FM
 
 		float OutputSizeAspectRatio = FullOutputSize.X / (float)FullOutputSize.Y;
 		const FIntPoint ConstrainedFullSize = CameraCache.AspectRatio > OutputSizeAspectRatio ?
-			FIntPoint(FullOutputSize.X, FMath::CeilToInt((double)FullOutputSize.X / (double)CameraCache.AspectRatio)) :
-			FIntPoint(FMath::CeilToInt(CameraCache.AspectRatio * FullOutputSize.Y), FullOutputSize.Y);
+			FIntPoint(FullOutputSize.X, FMath::RoundToInt((double)FullOutputSize.X / (double)CameraCache.AspectRatio)) :
+			FIntPoint(FMath::RoundToInt(CameraCache.AspectRatio * FullOutputSize.Y), FullOutputSize.Y);
 
 		const FIntPoint TileViewMin = InSampleState.OverlappedOffset;
 		const FIntPoint TileViewMax = TileViewMin + InSampleState.BackbufferSize;
