MoviePipeline: Updated the renderer nodes in the graph to hide the "Tone Curve" show flag, because there's already a "Disable Tone Curve" property.

diff --git a/Engine/Plugins/MovieScene/MovieRenderPipeline/Source/MovieRenderPipelineEditor/Private/Customizations/Graph/MovieGraphShowFlagsCustomization.h b/Engine/Plugins/MovieScene/MovieRenderPipeline/Source/MovieRenderPipelineEditor/Private/Customizations/Graph/MovieGraphShowFlagsCustomization.h
--- a/Engine/Plugins/MovieScene/MovieRenderPipeline/Source/MovieRenderPipelineEditor/Private/Customizations/Graph/MovieGraphShowFlagsCustomization.h
+++ b/Engine/Plugins/MovieScene/MovieRenderPipeline/Source/MovieRenderPipelineEditor/Private/Customizations/Graph/MovieGraphShowFlagsCustomization.h
@@ -45,6 +45,15 @@ protected:
 			{SFG_Advanced, LOCTEXT("GraphAdvancedSF", "Advanced")}
 		};
 
+		// Show flags that are excluded from showing up for any node
+		static const TMap<EShowFlagGroup, TArray<FName>> ExcludedShowFlagsByGroup =
+		{
+			// Renderer nodes already expose a "Disable Tone Curve" property, so the equivalent show flag is redundant
+			{
+				SFG_PostProcess, { FName(TEXT("ToneCurve")) }
+			}
+		};
+
 		check(PropertyHandle->IsValidHandle());
 
 		TArray<UObject*> OuterObjects;
@@ -96,6 +105,13 @@ protected:
 
 			for (const FShowFlagData& ShowFlag : *ShowFlags)
 			{
+				// This show flag may be excluded from being displayed
+				const TArray<FName>* ExcludedShowFlagsForGroup = ExcludedShowFlagsByGroup.Find(ShowFlag.Group);
+				if (ExcludedShowFlagsForGroup && ExcludedShowFlagsForGroup->Contains(ShowFlag.ShowFlagName))
+				{
+					continue;
+				}
+
 				const uint32 ShowFlagIndex = ShowFlag.EngineShowFlagIndex;
 
 				// Temporary workaround: ShowFlag.DisplayName can sometimes be blank, so the non-display name is used instead.
