MoviePipeline: Updated Quick Render's "Play Last Render" feature to remember the last render across editor sessions.

diff --git a/Engine/Plugins/MovieScene/MovieRenderPipeline/Source/MovieRenderPipelineEditor/Private/Graph/MovieGraphQuickRender.cpp b/Engine/Plugins/MovieScene/MovieRenderPipeline/Source/MovieRenderPipelineEditor/Private/Graph/MovieGraphQuickRender.cpp
--- a/Engine/Plugins/MovieScene/MovieRenderPipeline/Source/MovieRenderPipelineEditor/Private/Graph/MovieGraphQuickRender.cpp
+++ b/Engine/Plugins/MovieScene/MovieRenderPipeline/Source/MovieRenderPipelineEditor/Private/Graph/MovieGraphQuickRender.cpp
@@ -228,12 +228,12 @@ void UMovieGraphQuickRenderSubsystem::BeginQuickRender(const EMovieGraphQuickRen
 
 void UMovieGraphQuickRenderSubsystem::PlayLastRender()
 {
-	OpenPostRenderFileDisplayProcessor(PreviousRenderOutputData);
+	FMoviePipelinePostRenderFileDisplayProcessor::OpenPreviousFiles(GetDefault<UMovieRenderGraphEditorSettings>()->PostRenderSettings);
 }
 
 bool UMovieGraphQuickRenderSubsystem::CanPlayLastRender()
 {
-	return !PreviousRenderOutputData.GraphData.IsEmpty();
+	return FMoviePipelinePostRenderFileDisplayProcessor::CanOpenPreviousFiles();
 }
 
 void UMovieGraphQuickRenderSubsystem::OpenOutputDirectory(const UMovieGraphQuickRenderModeSettings* InQuickRenderSettings)
@@ -1352,8 +1352,6 @@ void UMovieGraphQuickRenderSubsystem::HandleJobFinished(const UMovieGraphQuickRe
 	{
 		OpenOutputDirectory(InQuickRenderSettings);
 	}
-
-	PreviousRenderOutputData = InGeneratedOutputData;
 }
 
 bool UMovieGraphQuickRenderSubsystem::IsViewportLookFlagActive(const EMovieGraphQuickRenderViewportLookFlags ViewportLookFlag) const
diff --git a/Engine/Plugins/MovieScene/MovieRenderPipeline/Source/MovieRenderPipelineEditor/Private/MoviePipelinePostRenderFileDisplayProcessor.cpp b/Engine/Plugins/MovieScene/MovieRenderPipeline/Source/MovieRenderPipelineEditor/Private/MoviePipelinePostRenderFileDisplayProcessor.cpp
--- a/Engine/Plugins/MovieScene/MovieRenderPipeline/Source/MovieRenderPipelineEditor/Private/MoviePipelinePostRenderFileDisplayProcessor.cpp
+++ b/Engine/Plugins/MovieScene/MovieRenderPipeline/Source/MovieRenderPipelineEditor/Private/MoviePipelinePostRenderFileDisplayProcessor.cpp
@@ -4,9 +4,13 @@
 
 #include "Graph/MovieRenderGraphEditorSettings.h"
 #include "Internationalization/Regex.h"
+#include "Misc/PackageName.h"
 #include "MovieGraphImageSequenceOutputNode.h"
 #include "MoviePipelineUtils.h"
 #include "MovieRenderPipelineCoreModule.h"
+#include "PackageHelperFunctions.h"
+#include "Settings/EditorLoadingSavingSettings.h"
+#include "UObject/Package.h"
 
 FMoviePipelinePostRenderFileDisplayProcessor::FMoviePipelinePostRenderFileDisplayProcessor(const FMovieGraphPostRenderSettings& InPostRenderSettings)
 	: PostRenderSettings(InPostRenderSettings)
@@ -54,7 +58,36 @@ void FMoviePipelinePostRenderFileDisplayProcessor::OpenFiles() const
 	if (!FinalListOfFilesToOpen.IsEmpty() && FinalPlayOptions)
 	{
 		LaunchFilesWithSettings(FinalListOfFilesToOpen, FinalPlayOptions, TTuple<FString, FString>(StartFrame, EndFrame));
+
+		// Persist the last set of pipeline output data so it can be played back later (even in another editor session). The raw output data is
+		// saved here, rather than the final list of files, because the user may subsequently change their playback settings, which could generate a
+		// different set of files to play back.
+		PersistOutputData(PipelineOutputData);
+	}
+}
+
+void FMoviePipelinePostRenderFileDisplayProcessor::OpenPreviousFiles(const FMovieGraphPostRenderSettings& InPostRenderSettings)
+{
+	if (!CanOpenPreviousFiles())
+	{
+		return;
+	}
+
+	FMoviePipelinePostRenderFileDisplayProcessor NewProcessor(InPostRenderSettings);
+
+	for (const FMoviePipelineOutputData& OutputData : GetPersistedOutputData()->OutputData)
+	{
+		NewProcessor.AddFiles(OutputData);
 	}
+
+	NewProcessor.OpenFiles();
+}
+
+bool FMoviePipelinePostRenderFileDisplayProcessor::CanOpenPreviousFiles()
+{
+	const UMoviePipelineOutputDataContainer* PersistedData = GetPersistedOutputData();
+
+	return PersistedData && !PersistedData->OutputData.IsEmpty();
 }
 
 FMoviePipelinePostRenderFileDisplayProcessor::FFilesToOpen FMoviePipelinePostRenderFileDisplayProcessor::FindFilesToOpen(FMoviePipelineOutputData& InPipelineOutputData) const
@@ -613,3 +646,63 @@ void FMoviePipelinePostRenderFileDisplayProcessor::GetStartAndEndFrames(const FS
 		}
 	}
 }
+
+void FMoviePipelinePostRenderFileDisplayProcessor::PersistOutputData(const TArray<FMoviePipelineOutputData>& InOutputData)
+{
+	// The output data is saved to a uasset rather than using the ini configuration system. The ini configuration system isn't flexible enough
+	// to store all of the types of data in FMoviePipelineOutputData. Additionally, the output data has the potential to contain hundreds/thousands of
+	// files, and an ini is not an appropriate storage location for that volume of data.
+	
+	static const FString PackageFileName = FPackageName::LongPackageNameToFilename(*PersistedOutputDataPackagePath, FPackageName::GetAssetPackageExtension());
+
+	// Get a pointer to the existing (or empty) persisted data in the target package (which will be saved to disk).
+	UMoviePipelineOutputDataContainer* PersistedData = GetPersistedOutputData();
+	if (!PersistedData)
+	{
+		return;
+	}
+
+	// Ensure the persisted data has the correct flags set so it actually gets saved to disk
+	PersistedData->OutputData = InOutputData;
+	PersistedData->SetFlags(RF_Public | RF_Standalone);
+
+	bool bSuccess = false;
+
+	// Save the settings out to disk. Turn off the behavior that auto-adds new files to source control.
+	{
+		UEditorLoadingSavingSettings* SaveSettings = GetMutableDefault<UEditorLoadingSavingSettings>();
+		const uint32 bSCCAutoAddNewFiles = SaveSettings->bSCCAutoAddNewFiles;
+		SaveSettings->bSCCAutoAddNewFiles = 0;
+
+		bSuccess = SavePackageHelper(PersistedData->GetPackage(), *PackageFileName);
+
+		SaveSettings->bSCCAutoAddNewFiles = bSCCAutoAddNewFiles;
+	}
+
+	if (!bSuccess)
+	{
+		// SavePackageHelper() will emit warnings if the save was unsuccessful, but log a separate warning for movie pipeline in case warnings are
+		// being specifically filtered for LogMovieRenderPipeline.
+		UE_LOG(LogMovieRenderPipeline, Warning, TEXT("Unable to save post-render play data. Could not save to destination file [%s]."), *PackageFileName);
+	}
+}
+
+UMoviePipelineOutputDataContainer* FMoviePipelinePostRenderFileDisplayProcessor::GetPersistedOutputData()
+{
+	UMoviePipelineOutputDataContainer* OutputData = nullptr;
+
+	if (const UPackage* OutputDataPackage = LoadPackage(nullptr, *PersistedOutputDataPackagePath, LOAD_None))
+	{
+		OutputData = Cast<UMoviePipelineOutputDataContainer>(FindObjectWithOuter(OutputDataPackage, UMoviePipelineOutputDataContainer::StaticClass()));
+	}
+
+	// The output data asset may not exist on disk yet, or (rarely) the object may be from a future version (eg, if the user went back to a
+	// previous build). Both of these situations require a new output data asset to be created.
+	if (!OutputData)
+	{
+		UPackage* NewOutputDataPackage = CreatePackage(*PersistedOutputDataPackagePath);
+		OutputData = NewObject<UMoviePipelineOutputDataContainer>(NewOutputDataPackage);
+	}
+
+	return OutputData;
+}
diff --git a/Engine/Plugins/MovieScene/MovieRenderPipeline/Source/MovieRenderPipelineEditor/Private/MoviePipelinePostRenderFileDisplayProcessor.h b/Engine/Plugins/MovieScene/MovieRenderPipeline/Source/MovieRenderPipelineEditor/Private/MoviePipelinePostRenderFileDisplayProcessor.h
--- a/Engine/Plugins/MovieScene/MovieRenderPipeline/Source/MovieRenderPipelineEditor/Private/MoviePipelinePostRenderFileDisplayProcessor.h
+++ b/Engine/Plugins/MovieScene/MovieRenderPipeline/Source/MovieRenderPipelineEditor/Private/MoviePipelinePostRenderFileDisplayProcessor.h
@@ -7,6 +7,20 @@
 #include "MoviePipelineQueue.h"
 #include "MoviePipelinePostRenderSettings.h"
 
+#include "MoviePipelinePostRenderFileDisplayProcessor.generated.h"
+
+/** UObject wrapper for the FMoviePipelineOutputData struct so it can be saved in a UPackage. */
+UCLASS()
+class UMoviePipelineOutputDataContainer : public UObject
+{
+	GENERATED_BODY()
+
+public:
+	/** The output data that the pipeline generated. Each entry in the array is the output data for one job. */
+	UPROPERTY()
+	TArray<FMoviePipelineOutputData> OutputData;
+};
+
 /** Handles sorting through all of the images/videos generated by a render(s), and opening them according to post-render editor preferences. */
 class FMoviePipelinePostRenderFileDisplayProcessor
 {
@@ -22,6 +36,15 @@ public:
 	/** Opens/displays rendered images/videos according to the post-render settings the class was initialized with. */
 	void OpenFiles() const;
 
+	/**
+	 * Open the files that OpenFiles() last opened, using the provided post-render settings (the previous files can be played back with
+	 * different settings than they were originally opened with).
+	 */
+	static void OpenPreviousFiles(const FMovieGraphPostRenderSettings& InPostRenderSettings);
+
+	/** Determines if OpenPreviousFiles() has any files that it can open. */
+	static bool CanOpenPreviousFiles();
+
 private:
 	/** Keeps track of the working set of images/videos that should be opened, and any other associated transient data. */
 	struct FFilesToOpen
@@ -76,7 +99,19 @@ private:
 	/** Gets the start and end frame of the images provided, based on the frame-templated path. The string and int-based frame numbers in InOutFilesToOpen are updated. */
 	void GetStartAndEndFrames(const FString& InTemplatedPath, const TArray<FString>& InImagePaths, FFilesToOpen& InOutFilesToOpen) const;
 
+	/** Persists the provided pipeline output data (one array entry per job) to disk at /Temp/MovieRenderPipeline/MostRecentOutputData. */
+	static void PersistOutputData(const TArray<FMoviePipelineOutputData>& InOutputData);
+
+	/**
+	 * Gets the most recently persisted pipeline output data (see PersistOutputData()). Note that if data has not been persisted yet, this method
+	 * will still return a valid pointer to an output data container (however, it will not contain any file paths within it).
+	 */
+	static UMoviePipelineOutputDataContainer* GetPersistedOutputData();
+
 private:
+	/** The package path for the output data persisted across editor sessions (see PersistOutputData()). */
+	inline static const FString PersistedOutputDataPackagePath = TEXT("/Temp/MovieRenderPipeline/MostRecentOutputData");
+
 	/** The post-render settings the class was initialized with. Usually sourced from the editor preferences. */
 	FMovieGraphPostRenderSettings PostRenderSettings;
 
diff --git a/Engine/Plugins/MovieScene/MovieRenderPipeline/Source/MovieRenderPipelineEditor/Public/Graph/MovieGraphQuickRender.h b/Engine/Plugins/MovieScene/MovieRenderPipeline/Source/MovieRenderPipelineEditor/Public/Graph/MovieGraphQuickRender.h
--- a/Engine/Plugins/MovieScene/MovieRenderPipeline/Source/MovieRenderPipelineEditor/Public/Graph/MovieGraphQuickRender.h
+++ b/Engine/Plugins/MovieScene/MovieRenderPipeline/Source/MovieRenderPipelineEditor/Public/Graph/MovieGraphQuickRender.h
@@ -215,9 +215,6 @@ private:
 
 	/** Any data that needed to be cached prior to starting PIE. */
 	FCachedPrePieData CachedPrePieData;
-
-	/** Cache of the output data from the last render. */
-	FMoviePipelineOutputData PreviousRenderOutputData;
 };
 
 #undef UE_API
