Movie Pipeline: Fix CIS

#rnx

diff --git a/Engine/Plugins/MovieScene/MovieRenderPipeline/Source/MovieRenderPipelineRenderPasses/Private/MovieGraphImageSequenceOutputNode.cpp b/Engine/Plugins/MovieScene/MovieRenderPipeline/Source/MovieRenderPipelineRenderPasses/Private/MovieGraphImageSequenceOutputNode.cpp
--- a/Engine/Plugins/MovieScene/MovieRenderPipeline/Source/MovieRenderPipelineRenderPasses/Private/MovieGraphImageSequenceOutputNode.cpp
+++ b/Engine/Plugins/MovieScene/MovieRenderPipeline/Source/MovieRenderPipelineRenderPasses/Private/MovieGraphImageSequenceOutputNode.cpp
@@ -13,6 +13,7 @@
 #include "Graph/MovieRenderGraphEditorSettings.h"
 #include "MoviePipelineUtils.h"
 #include "MoviePipelineImageSequenceOutput.h" // for FAsyncImageQuantization
+#include "MoviePipelineEXROutput.h"
 #include "MovieRenderPipelineCoreModule.h"
 
 #include "Modules/ModuleManager.h"
