Movie Pipeline: Fix CIS Wrap usages of FEXRImageWriteTask in WIth_UNREALEXR

#rnx

diff --git a/Engine/Plugins/MovieScene/MovieRenderPipeline/Source/MovieRenderPipelineRenderPasses/Private/MovieGraphImageSequenceOutputNode.cpp b/Engine/Plugins/MovieScene/MovieRenderPipeline/Source/MovieRenderPipelineRenderPasses/Private/MovieGraphImageSequenceOutputNode.cpp
--- a/Engine/Plugins/MovieScene/MovieRenderPipeline/Source/MovieRenderPipelineRenderPasses/Private/MovieGraphImageSequenceOutputNode.cpp
+++ b/Engine/Plugins/MovieScene/MovieRenderPipeline/Source/MovieRenderPipelineRenderPasses/Private/MovieGraphImageSequenceOutputNode.cpp
@@ -287,6 +287,7 @@ void UMovieGraphImageSequenceOutputNode::OnReceiveImageDataImpl(UMovieGraphPipel
 	}
 }
 
+#if WITH_UNREALEXR
 
 TUniquePtr<FEXRImageWriteTask> UMovieGraphImageSequenceOutputNode_EXR::CreateImageWriteTask(FString InFileName, EEXRCompressionFormat InCompression, bool bMultiPart) const
 {
@@ -641,6 +642,8 @@ void UMovieGraphImageSequenceOutputNode_MultiLayerEXR::OnReceiveImageDataImpl(UM
 	}
 }
 
+#endif // WITH_UNREALEXR
+
 void UMovieGraphImageSequenceOutputNode_MultiLayerEXR::GetFilenameToEXROutputConfigMappings(
 	const UMovieGraphImageSequenceOutputNode_MultiLayerEXR* InParentNode, UMovieGraphPipeline* InPipeline,
 	UE::MovieGraph::FMovieGraphOutputMergerFrame* InRawFrameData, TMap<FString, FEXROutputConfigForFilename>& OutFilenameToOutputConfigs) const
diff --git a/Engine/Plugins/MovieScene/MovieRenderPipeline/Source/MovieRenderPipelineRenderPasses/Public/MovieGraphImageSequenceOutputNode.h b/Engine/Plugins/MovieScene/MovieRenderPipeline/Source/MovieRenderPipelineRenderPasses/Public/MovieGraphImageSequenceOutputNode.h
--- a/Engine/Plugins/MovieScene/MovieRenderPipeline/Source/MovieRenderPipelineRenderPasses/Public/MovieGraphImageSequenceOutputNode.h
+++ b/Engine/Plugins/MovieScene/MovieRenderPipeline/Source/MovieRenderPipelineRenderPasses/Public/MovieGraphImageSequenceOutputNode.h
@@ -107,7 +107,10 @@ public:
 		Compression = EEXRCompressionFormat::PIZ;
 	}
 
+#if WITH_UNREALEXR
 	UE_API virtual void OnReceiveImageDataImpl(UMovieGraphPipeline* InPipeline, UE::MovieGraph::FMovieGraphOutputMergerFrame* InRawFrameData, const TSet<FMovieGraphRenderDataIdentifier>& InMask) override;
+#endif // WITH_UNREALEXR
+
 	virtual bool ShouldCropOverscanImpl() const override { return false; }
 	
 #if WITH_EDITOR
@@ -149,6 +152,8 @@ public:
 
 protected:
 
+#if WITH_UNREALEXR
+
 	/** Convenience function to create a new EXR image write task, given a file name and compression format. */
 	UE_API TUniquePtr<FEXRImageWriteTask> CreateImageWriteTask(
 		FString InFileName,
@@ -172,6 +177,8 @@ protected:
 		const FString& InLayerName = {},
 		const TMap<FString, FString>& InResolvedOCIOContext = {}
 	) const;
+
+#endif // WITH_UNREALEXR
 };
 
 
@@ -197,7 +204,10 @@ public:
 		return EMovieGraphBranchRestriction::Globals;
 	}
 
+#if WITH_UNREALEXR
 	UE_API virtual void OnReceiveImageDataImpl(UMovieGraphPipeline* InPipeline, UE::MovieGraph::FMovieGraphOutputMergerFrame* InRawFrameData, const TSet<FMovieGraphRenderDataIdentifier>& InMask) override;
+#endif // WITH_UNREALEXR
+
 	virtual bool ShouldCropOverscanImpl() const override { return false; }
 	
 #if WITH_EDITOR
diff --git a/Engine/Plugins/MovieScene/MovieRenderPipeline/Source/MovieRenderPipelineRenderPasses/Public/MoviePipelineEXROutput.h b/Engine/Plugins/MovieScene/MovieRenderPipeline/Source/MovieRenderPipelineRenderPasses/Public/MoviePipelineEXROutput.h
--- a/Engine/Plugins/MovieScene/MovieRenderPipeline/Source/MovieRenderPipelineRenderPasses/Public/MoviePipelineEXROutput.h
+++ b/Engine/Plugins/MovieScene/MovieRenderPipeline/Source/MovieRenderPipelineRenderPasses/Public/MoviePipelineEXROutput.h
@@ -202,7 +202,9 @@ public:
 		bMultipart = false;
 	}
 
+#if WITH_UNREALEXR
 	MOVIERENDERPIPELINERENDERPASSES_API virtual void OnReceiveImageDataImpl(FMoviePipelineMergerOutputFrame* InMergedOutputFrame) override;
+#endif // WITH_UNREALEXR
 	virtual bool IsAlphaAllowed() const override { return true; }
 	virtual bool ShouldCropOverscanImpl() const override { return false; }
 
