MoviePipeline: Fixed an issue where the "View" menu in the graph's editor could cause a crash if the graph is reloaded, and the level has changed since the editor's initial load.

diff --git a/Engine/Plugins/MovieScene/MovieRenderPipeline/Source/MovieRenderPipelineEditor/Private/Graph/MovieGraphAssetToolkit.cpp b/Engine/Plugins/MovieScene/MovieRenderPipeline/Source/MovieRenderPipelineEditor/Private/Graph/MovieGraphAssetToolkit.cpp
--- a/Engine/Plugins/MovieScene/MovieRenderPipeline/Source/MovieRenderPipelineEditor/Private/Graph/MovieGraphAssetToolkit.cpp
+++ b/Engine/Plugins/MovieScene/MovieRenderPipeline/Source/MovieRenderPipelineEditor/Private/Graph/MovieGraphAssetToolkit.cpp
@@ -545,11 +545,11 @@ void FMovieGraphAssetToolkit::ExtendToolkitMenu() const
 				"View",
 				LOCTEXT("ViewMenu", "View"),
 				LOCTEXT("ViewMenu_ToolTip", "Open the View menu"),
-				FNewToolMenuDelegate::CreateLambda([this](UToolMenu* InMenu)
+				FNewToolMenuDelegate::CreateLambda([](UToolMenu* InMenu)
 				{
 					FToolMenuSection& ZoomSection = InMenu->AddSection("ViewZoom", LOCTEXT("ViewMenuZoomHeading", "Zoom"));
-					ZoomSection.AddMenuEntryWithCommandList(FMoviePipelineCommands::Get().ZoomToWindow, ToolkitCommands);
-					ZoomSection.AddMenuEntryWithCommandList(FMoviePipelineCommands::Get().ZoomToSelection, ToolkitCommands);
+					ZoomSection.AddMenuEntry(FMoviePipelineCommands::Get().ZoomToWindow);
+					ZoomSection.AddMenuEntry(FMoviePipelineCommands::Get().ZoomToSelection);
 				})
 			).InsertPosition = FToolMenuInsert("Asset", EToolMenuInsertType::After);
 		}
