MoviePipeline: Fixed an issue where the movie pipeline could crash when a remote render is stopped midway through a render.

diff --git a/Engine/Plugins/MovieScene/MovieRenderPipeline/Source/MovieRenderPipelineCore/Private/MoviePipelineInProcessExecutor.cpp b/Engine/Plugins/MovieScene/MovieRenderPipeline/Source/MovieRenderPipelineCore/Private/MoviePipelineInProcessExecutor.cpp
--- a/Engine/Plugins/MovieScene/MovieRenderPipeline/Source/MovieRenderPipelineCore/Private/MoviePipelineInProcessExecutor.cpp
+++ b/Engine/Plugins/MovieScene/MovieRenderPipeline/Source/MovieRenderPipelineCore/Private/MoviePipelineInProcessExecutor.cpp
@@ -139,7 +139,12 @@ void UMoviePipelineInProcessExecutor::OnMapLoadFinished(UWorld* NewWorld)
 	
 	// Listen for when the pipeline thinks it has finished.
 	ActiveMoviePipeline->OnMoviePipelineWorkFinished().AddUObject(this, &UMoviePipelineInProcessExecutor::OnMoviePipelineFinished);
-	FCoreDelegates::OnEnginePreExit.AddUObject(this, &UMoviePipelineInProcessExecutor::OnApplicationQuit);
+
+	// Give the pipeline a chance to shut down properly if the render is interrupted (Alt+F4, closing the preview window, etc). When the render is
+	// stopped like this, the engine's shutdown process may kill the player before more common delegates like FCoreDelegates::OnEnginePreExit are
+	// called. ApplicationWillTerminateDelegate will be called immediately after an engine shutdown request is received though, and lets the pipeline
+	// shut down before trying to access the (now non-existent) player.
+	FCoreDelegates::GetApplicationWillTerminateDelegate().AddUObject(this, &UMoviePipelineInProcessExecutor::OnApplicationQuit);
 
 	// Wait until we actually recieved the right map and created the pipeline before saying that we're actively rendering
 	bIsRendering = true;
