MoviePipeline: Added Naming Tokens support to MRG.

Also did some slight optimizations in the Naming Tokens system.

diff --git a/Engine/Plugins/Developer/NamingTokens/Source/NamingTokens/Private/Utils/NamingTokenUtils.cpp b/Engine/Plugins/Developer/NamingTokens/Source/NamingTokens/Private/Utils/NamingTokenUtils.cpp
--- a/Engine/Plugins/Developer/NamingTokens/Source/NamingTokens/Private/Utils/NamingTokenUtils.cpp
+++ b/Engine/Plugins/Developer/NamingTokens/Source/NamingTokens/Private/Utils/NamingTokenUtils.cpp
@@ -40,15 +40,17 @@ FString UE::NamingTokens::Utils::CreateFormattedToken(const FNamingTokenData& In
 
 FString UE::NamingTokens::Utils::GetNamespaceDelimiter()
 {
-	return TEXT(":");
+	static const FString Delimiter(TEXT(":"));
+
+	return Delimiter;
 }
 
 TArray<FString> UE::NamingTokens::Utils::GetTokenKeysFromString(const FString& InTokenString)
 {
-	TSet<FString, FLocKeySetFuncs> Tokens;
+	static const FString PatternString = Private::GetTokenPatternString();
+	static const FRegexPattern Pattern(PatternString);
 
-	const FString PatternString = Private::GetTokenPatternString();
-	const FRegexPattern Pattern(PatternString);
+	TSet<FString, FLocKeySetFuncs> Tokens;
 	FRegexMatcher Matcher(Pattern, InTokenString);
 
 	while (Matcher.FindNext())
@@ -62,8 +64,9 @@ TArray<FString> UE::NamingTokens::Utils::GetTokenKeysFromString(const FString& I
 
 bool UE::NamingTokens::Utils::IsTokenInString(const FString& InTokenKey, const FString& InTokenString)
 {
-	const FString PatternString = Private::GetTokenPatternString();
-	const FRegexPattern Pattern(PatternString);
+	static const FString PatternString = Private::GetTokenPatternString();
+	static const FRegexPattern Pattern(PatternString);
+
 	FRegexMatcher Matcher(Pattern, InTokenString);
 
 	while (Matcher.FindNext())
diff --git a/Engine/Plugins/MovieScene/MovieRenderPipeline/MovieRenderPipeline.uplugin b/Engine/Plugins/MovieScene/MovieRenderPipeline/MovieRenderPipeline.uplugin
--- a/Engine/Plugins/MovieScene/MovieRenderPipeline/MovieRenderPipeline.uplugin
+++ b/Engine/Plugins/MovieScene/MovieRenderPipeline/MovieRenderPipeline.uplugin
@@ -46,6 +46,11 @@
 		{
 			"Name": "ChaosClothAsset",
 			"Enabled": true
+		},
+		{
+			"Name": "NamingTokens",
+			"Enabled": true,
+			"Optional": false
 		}
 	],
 	"Modules": [
diff --git a/Engine/Plugins/MovieScene/MovieRenderPipeline/Source/MovieRenderPipelineCore/MovieRenderPipelineCore.Build.cs b/Engine/Plugins/MovieScene/MovieRenderPipeline/Source/MovieRenderPipelineCore/MovieRenderPipelineCore.Build.cs
--- a/Engine/Plugins/MovieScene/MovieRenderPipeline/Source/MovieRenderPipelineCore/MovieRenderPipelineCore.Build.cs
+++ b/Engine/Plugins/MovieScene/MovieRenderPipeline/Source/MovieRenderPipelineCore/MovieRenderPipelineCore.Build.cs
@@ -27,6 +27,7 @@ public class MovieRenderPipelineCore : ModuleRules
 				"ImageCore",
 				"ChaosClothAssetEngine",
 				"Renderer",
+				"NamingTokens",
 			}
 		);
 
diff --git a/Engine/Plugins/MovieScene/MovieRenderPipeline/Source/MovieRenderPipelineCore/Private/Graph/MovieGraphBlueprintLibrary.cpp b/Engine/Plugins/MovieScene/MovieRenderPipeline/Source/MovieRenderPipelineCore/Private/Graph/MovieGraphBlueprintLibrary.cpp
--- a/Engine/Plugins/MovieScene/MovieRenderPipeline/Source/MovieRenderPipelineCore/Private/Graph/MovieGraphBlueprintLibrary.cpp
+++ b/Engine/Plugins/MovieScene/MovieRenderPipeline/Source/MovieRenderPipelineCore/Private/Graph/MovieGraphBlueprintLibrary.cpp
@@ -8,18 +8,95 @@
 #include "Graph/Nodes/MovieGraphCommandLineEncoderNode.h"
 #include "Graph/Nodes/MovieGraphFileOutputNode.h"
 #include "Graph/Nodes/MovieGraphGlobalOutputSettingNode.h"
-#include "Graph/Nodes/MovieGraphRenderLayerNode.h"
 #include "MoviePipelineBlueprintLibrary.h"
 #include "MoviePipelineUtils.h"
-#include "CineCameraActor.h"
 #include "MovieRenderPipelineCoreModule.h"
 
-#include "Misc/Paths.h"
+#include "CineCameraActor.h"
+#include "GlobalNamingTokens.h"
 #include "HAL/FileManager.h"
 #include "Internationalization/Regex.h"
+#include "Misc/Paths.h"
+#include "NamingTokensEngineSubsystem.h"
+#include "Utils/NamingTokenUtils.h"
 
 #include UE_INLINE_GENERATED_CPP_BY_NAME(MovieGraphBlueprintLibrary)
 
+namespace
+{
+	/**
+	 * From the provided format string, resolve all of the non-MRG {tokens} from the Naming Tokens system. Resolved values are provided in both the
+	 * InOutFilenameArguments and InOutFormatArguments maps. The former is exposed to the API, and the latter is just used for FString::Format().
+	 *
+	 * Note that tokens in the Naming Tokens system are usually namespaced, so the tokens added to InOutFormatArguments will be renamed to be
+	 * compatible with FString::Format() (to eliminate the namespacing colon delimiter). Tokens added to InOutFilenameArguments will not have their
+	 * names modified. InOutFormatString will be modified in-place to contain the updated token names that InOutFormatArguments uses.
+	 */
+	void UpdateTokenMapsWithResolvedExternalTokens(
+		FString& InOutFormatString, TMap<FString, FString>& InOutFilenameArguments, FStringFormatNamedArguments& InOutFormatArguments)
+	{
+		// With the given token name and evaluated token data, updates the format string as needed and updates the format and filename arguments with
+		// the evaluated token value.
+		auto UpdateFormatStringAndArguments = [&InOutFormatString, &InOutFilenameArguments, &InOutFormatArguments](
+			const FString& InTokenName, const FNamingTokenValueData& InEvaluatedTokenData)
+		{
+			// FString::Format() does not work with tokens that contain colons, so generate a temporary replacement token and do a find-and-replace
+			// on the source format string to use the new substitute token name. Eg, "tr:slate" -> "__MRGTEMP_tr_slate__".
+			const FString TempMrgTokenName = FString::Printf(TEXT("__MRGTEMP_%s_%s__"), *InEvaluatedTokenData.TokenNamespace, *InEvaluatedTokenData.TokenKey);
+
+			// Replace the existing {token} with the temporary token name generated above.
+			InOutFormatString.ReplaceInline(
+				*FString::Printf(TEXT("{%s}"), *InTokenName),
+				*FString::Printf(TEXT("{%s}"), *TempMrgTokenName));
+
+			// Add the modified token to the format arguments, but use the un-modified token name in the resolved arguments (which is visible to API consumers).
+			const FString EvaluatedTokenValue = InEvaluatedTokenData.TokenValue.ToString();
+			InOutFormatArguments.Add(TempMrgTokenName, EvaluatedTokenValue);
+			InOutFilenameArguments.Add(InTokenName, EvaluatedTokenValue);
+		};
+		
+		// Parse out all of the filename argument {tokens} in the string, and find the ones that do not come from MRG.
+		TArray<FString> NonMrgTokenNames;
+		for (const FString& FilenameArgument : UE::NamingTokens::Utils::GetTokenKeysFromString(InOutFormatString))
+		{
+			if (!InOutFilenameArguments.Contains(FilenameArgument))
+			{
+				NonMrgTokenNames.Add(FilenameArgument);
+			}
+		}
+
+		// Can't resolve naming tokens w/o the Naming Tokens plugin loaded. We can't really warn here at this point because even if some non-MRG
+		// tokens were found, we can't say definitively that they're defined in the Naming Tokens system (eg, the token could just be misspelled).
+		UNamingTokensEngineSubsystem* NamingTokensSubsystem = GEngine->GetEngineSubsystem<UNamingTokensEngineSubsystem>();
+		if (!NamingTokensSubsystem)
+		{
+			return;
+		}
+
+		// Resolve via the Naming Tokens system and inject the values into the argument maps.
+		for (const FNamingTokenValueData& EvaluatedTokenData : NamingTokensSubsystem->EvaluateTokenList(NonMrgTokenNames))
+		{
+			// If the token could not be resolved, skip it
+			if (!EvaluatedTokenData.bWasEvaluated)
+			{
+				continue;
+			}
+
+			const FString NamespacedToken = UE::NamingTokens::Utils::CombineNamespaceAndTokenKey(EvaluatedTokenData.TokenNamespace, EvaluatedTokenData.TokenKey);
+			if (NonMrgTokenNames.Contains(NamespacedToken))
+			{
+				UpdateFormatStringAndArguments(NamespacedToken, EvaluatedTokenData);
+			}
+
+			const FString NonNamespacedToken = EvaluatedTokenData.TokenKey;
+			if (NonMrgTokenNames.Contains(NonNamespacedToken))
+			{
+				UpdateFormatStringAndArguments(NonNamespacedToken, EvaluatedTokenData);
+			}
+		}
+	}
+}
+
 FFrameRate UMovieGraphBlueprintLibrary::GetEffectiveFrameRate(UMovieGraphGlobalOutputSettingNode* InNode, const FFrameRate& InDefaultRate)
 {
 	if (InNode && InNode->bOverride_OutputFrameRate)
@@ -36,6 +113,11 @@ FString UMovieGraphBlueprintLibrary::ResolveFilenameFormatArguments(const FStrin
 	// UMovieGraphPipeline - these are handled as resolve params to allow resolving during edit time
 	// UMoviePipelineExecutorJob - author, sequence, level, etc.
 	// UMovieGraphSettingNode - every node can contribute it's own file/metadata KVP based on their own settings
+
+	// Keep track of the resolve args that are generated JUST for this call to ResolveFilenameFormatArguments(). This generates a record of all tokens
+	// that were generated from MRG, and thus allows us to determine which tokens originate from outside MRG. These will be merged with
+	// OutMergedFormatArgs before the function returns.
+	FMovieGraphResolveArgs GeneratedResolveArgs;
 	
 	bool bOverwriteExisting = true;
 
@@ -70,21 +152,21 @@ FString UMovieGraphBlueprintLibrary::ResolveFilenameFormatArguments(const FStrin
 		ShotName = ShotName.Len() > 0 ? ShotName : TEXT("NoShot");
 
 		// Use the legacy system to add filename and metadata kvp
-		MoviePipeline::GetOutputStateFormatArgs(OutMergedFormatArgs.FilenameArguments, OutMergedFormatArgs.FileMetadata, FrameNumber, FrameNumberShot, FrameNumberRel, FrameNumberShotRel, TEXT("DummyCameraToken"), ShotName);
+		MoviePipeline::GetOutputStateFormatArgs(GeneratedResolveArgs.FilenameArguments, GeneratedResolveArgs.FileMetadata, FrameNumber, FrameNumberShot, FrameNumberRel, FrameNumberShotRel, TEXT("DummyCameraToken"), ShotName);
 
 		// Add on Render Data Identifier, overwriting the dummy camera name above.
-		OutMergedFormatArgs.FilenameArguments.Add(TEXT("camera_name"), InParams.RenderDataIdentifier.CameraName);
-		OutMergedFormatArgs.FilenameArguments.Add(TEXT("renderer_name"), InParams.RenderDataIdentifier.RendererName);
-		OutMergedFormatArgs.FilenameArguments.Add(TEXT("renderer_sub_name"), InParams.RenderDataIdentifier.SubResourceName);
-		OutMergedFormatArgs.FilenameArguments.Add(TEXT("layer_name"), InParams.RenderDataIdentifier.LayerName);
-		OutMergedFormatArgs.FilenameArguments.Add(TEXT("branch_name"), InParams.RenderDataIdentifier.RootBranchName.ToString());
+		GeneratedResolveArgs.FilenameArguments.Add(TEXT("camera_name"), InParams.RenderDataIdentifier.CameraName);
+		GeneratedResolveArgs.FilenameArguments.Add(TEXT("renderer_name"), InParams.RenderDataIdentifier.RendererName);
+		GeneratedResolveArgs.FilenameArguments.Add(TEXT("renderer_sub_name"), InParams.RenderDataIdentifier.SubResourceName);
+		GeneratedResolveArgs.FilenameArguments.Add(TEXT("layer_name"), InParams.RenderDataIdentifier.LayerName);
+		GeneratedResolveArgs.FilenameArguments.Add(TEXT("branch_name"), InParams.RenderDataIdentifier.RootBranchName.ToString());
 
 		// TODO: Some of these are per render layer and need to be stored that way. EXRs will have the metadata for all the layers/cameras/etc. in one file.
-		OutMergedFormatArgs.FileMetadata.Add(TEXT("unreal/cameraName"), InParams.RenderDataIdentifier.CameraName);
-		OutMergedFormatArgs.FileMetadata.Add(TEXT("unreal/rendererName"), InParams.RenderDataIdentifier.RendererName);
-		OutMergedFormatArgs.FileMetadata.Add(TEXT("unreal/subResourceName"), InParams.RenderDataIdentifier.SubResourceName);
-		OutMergedFormatArgs.FileMetadata.Add(TEXT("unreal/layerName"), InParams.RenderDataIdentifier.LayerName);
-		OutMergedFormatArgs.FileMetadata.Add(TEXT("unreal/branchName"), InParams.RenderDataIdentifier.RootBranchName.ToString());
+		GeneratedResolveArgs.FileMetadata.Add(TEXT("unreal/cameraName"), InParams.RenderDataIdentifier.CameraName);
+		GeneratedResolveArgs.FileMetadata.Add(TEXT("unreal/rendererName"), InParams.RenderDataIdentifier.RendererName);
+		GeneratedResolveArgs.FileMetadata.Add(TEXT("unreal/subResourceName"), InParams.RenderDataIdentifier.SubResourceName);
+		GeneratedResolveArgs.FileMetadata.Add(TEXT("unreal/layerName"), InParams.RenderDataIdentifier.LayerName);
+		GeneratedResolveArgs.FileMetadata.Add(TEXT("unreal/branchName"), InParams.RenderDataIdentifier.RootBranchName.ToString());
 	}
 
 
@@ -104,8 +186,8 @@ FString UMovieGraphBlueprintLibrary::ResolveFilenameFormatArguments(const FStrin
 			}
 		}
 
-		OutMergedFormatArgs.FilenameArguments.Add(TEXT("frame_rate"), FString::SanitizeFloat(FrameRate));
-		OutMergedFormatArgs.FileMetadata.Add(TEXT("unreal/frameRate"), FString::SanitizeFloat(FrameRate));
+		GeneratedResolveArgs.FilenameArguments.Add(TEXT("frame_rate"), FString::SanitizeFloat(FrameRate));
+		GeneratedResolveArgs.FileMetadata.Add(TEXT("unreal/frameRate"), FString::SanitizeFloat(FrameRate));
 
 		FString LevelName;
 		FString SequenceName;
@@ -115,14 +197,14 @@ FString UMovieGraphBlueprintLibrary::ResolveFilenameFormatArguments(const FStrin
 			SequenceName = InParams.Job->Sequence.GetAssetName();
 		}
 		
-		OutMergedFormatArgs.FilenameArguments.Add(TEXT("level_name"), LevelName);
-		OutMergedFormatArgs.FilenameArguments.Add(TEXT("sequence_name"), SequenceName);
+		GeneratedResolveArgs.FilenameArguments.Add(TEXT("level_name"), LevelName);
+		GeneratedResolveArgs.FilenameArguments.Add(TEXT("sequence_name"), SequenceName);
 		
-		OutMergedFormatArgs.FileMetadata.Add(TEXT("unreal/levelName"), LevelName);
-		OutMergedFormatArgs.FileMetadata.Add(TEXT("unreal/sequenceName"), SequenceName);
+		GeneratedResolveArgs.FileMetadata.Add(TEXT("unreal/levelName"), LevelName);
+		GeneratedResolveArgs.FileMetadata.Add(TEXT("unreal/sequenceName"), SequenceName);
 		
 		// Add KVP data from the job (date, time, job name, job author, job comment)
-		UE::MoviePipeline::GetSharedFormatArguments(OutMergedFormatArgs.FilenameArguments, OutMergedFormatArgs.FileMetadata, InParams.InitializationTime, InParams.Version, InParams.Job, InParams.InitializationTimeOffset);
+		UE::MoviePipeline::GetSharedFormatArguments(GeneratedResolveArgs.FilenameArguments, GeneratedResolveArgs.FileMetadata, InParams.InitializationTime, InParams.Version, InParams.Job, InParams.InitializationTimeOffset);
 	}
 	
 
@@ -149,34 +231,34 @@ FString UMovieGraphBlueprintLibrary::ResolveFilenameFormatArguments(const FStrin
 			SettingInstances.Add(GetMutableDefault<UMovieGraphSettingNode>(InClass));
 		}
 
-		for(UMovieGraphSettingNode* Instance : SettingInstances)
+		for (UMovieGraphSettingNode* Instance : SettingInstances)
 		{
-			Instance->GetFormatResolveArgs(OutMergedFormatArgs, InParams.RenderDataIdentifier);
+			Instance->GetFormatResolveArgs(GeneratedResolveArgs, InParams.RenderDataIdentifier);
 		}
 	}
 
 	// Copy the metadata in the parameters into the merged output version
 	for (const TPair<FString, FString>& FileMetadata : InParams.FileMetadata)
 	{
-		OutMergedFormatArgs.FileMetadata.Add(FileMetadata.Key, FileMetadata.Value);
+		GeneratedResolveArgs.FileMetadata.Add(FileMetadata.Key, FileMetadata.Value);
 	}
 
 	// We expect the incoming string to have a {file_dup} token where they want file duplication numbers to be handled
 	// but we don't actually want a value there (unless there's a collision). So we override it with an empty string by default.
-	OutMergedFormatArgs.FilenameArguments.Add(TEXT("file_dup"), FString());
+	GeneratedResolveArgs.FilenameArguments.Add(TEXT("file_dup"), FString());
 
 	// Format the {version} token so it's of the form "v###"
-	OutMergedFormatArgs.FilenameArguments.Add(TEXT("version"), FString::Printf(TEXT("v%0*d"), 3, InParams.Version));
+	GeneratedResolveArgs.FilenameArguments.Add(TEXT("version"), FString::Printf(TEXT("v%0*d"), 3, InParams.Version));
 
 	// Overwrite any {tokens} with the user-supplied overrides if needed. This allows different requesters to share the same variables (ie: filename extension)
 	for (const TPair<FString, FString>& KVP : InParams.FileNameFormatOverrides)
 	{
-		OutMergedFormatArgs.FilenameArguments.Add(KVP.Key, KVP.Value);
+		GeneratedResolveArgs.FilenameArguments.Add(KVP.Key, KVP.Value);
 	}
 
 	// Convert from our Python/BP exposed <FString, FString> to a named argument array for the formatter.
 	FStringFormatNamedArguments NamedArgs;
-	for (const TPair<FString, FString>& Argument : OutMergedFormatArgs.FilenameArguments)
+	for (const TPair<FString, FString>& Argument : GeneratedResolveArgs.FilenameArguments)
 	{
 		NamedArgs.Add(Argument.Key, Argument.Value);
 	}
@@ -187,8 +269,18 @@ FString UMovieGraphBlueprintLibrary::ResolveFilenameFormatArguments(const FStrin
 		return FString();
 	}
 
+	// The format string may need to be modified from this point onwards, so make a copy so the source string isn't unintentionally changed.
+	FString FormatStringCopy = InFormatString;
+
+	// Resolve non-MRG tokens via the Naming Tokens system.
+	UpdateTokenMapsWithResolvedExternalTokens(FormatStringCopy, GeneratedResolveArgs.FilenameArguments, NamedArgs);
+
+	// Now that all format args have been discovered, add them to the merged format args.
+	OutMergedFormatArgs.FileMetadata.Append(GeneratedResolveArgs.FileMetadata);
+	OutMergedFormatArgs.FilenameArguments.Append(GeneratedResolveArgs.FilenameArguments);
+
 	// Apply all of our named args to the file and generate a path.
-	FString BaseFilename = FString::Format(*InFormatString, NamedArgs);
+	FString BaseFilename = FString::Format(*FormatStringCopy, NamedArgs);
 	if (InParams.bEnsureAbsolutePath)
 	{
 		BaseFilename = FPaths::ConvertRelativePathToFull(BaseFilename);
@@ -245,7 +337,7 @@ FString UMovieGraphBlueprintLibrary::ResolveFilenameFormatArguments(const FStrin
 		NamedArgs.Add(TEXT("file_dup"), FString::Printf(TEXT("_(%d)"), DuplicateIndex));
 
 		// Re-resolve the format string now that we've reassigned frame_dup to a number.
-		ThisTry = FString::Format(*InFormatString, NamedArgs) + ResolvedExtension;
+		ThisTry = FString::Format(*FormatStringCopy, NamedArgs) + ResolvedExtension;
 
 		// If the file doesn't exist, we can use that, else, increment the index and try again
 		if (UE::MoviePipeline::CanWriteToFile(*ThisTry, bOverwriteExisting))
