MoviePipeline: Added the ability to specify the lossless compression type on EXR nodes within the graph, separately from the default compression type.

There are some instances where you'd want to set up your EXR with lossy compression by default, but use lossless compression for data layers that need high precision. This is now possible; Object ID passes will always use the lossless compression type specified on the EXR node, and PPM passes on the deferred/pathtracer nodes can opt-in to using lossless compression ("Use Lossless Compression").

diff --git a/Engine/Plugins/MovieScene/MoviePipelineMaskRenderPass/Source/MoviePipelineMaskRenderPass/Private/Graph/MovieGraphObjectIdNode.cpp b/Engine/Plugins/MovieScene/MoviePipelineMaskRenderPass/Source/MoviePipelineMaskRenderPass/Private/Graph/MovieGraphObjectIdNode.cpp
--- a/Engine/Plugins/MovieScene/MoviePipelineMaskRenderPass/Source/MoviePipelineMaskRenderPass/Private/Graph/MovieGraphObjectIdNode.cpp
+++ b/Engine/Plugins/MovieScene/MoviePipelineMaskRenderPass/Source/MoviePipelineMaskRenderPass/Private/Graph/MovieGraphObjectIdNode.cpp
@@ -40,6 +40,12 @@ bool UMovieGraphObjectIdNode::GetAllowsCompositing() const
 	return false;
 }
 
+bool UMovieGraphObjectIdNode::GetForceLosslessCompression() const
+{
+	// Object ID data is highly sensitive to lossy compression; it should always use lossless compression.
+	return true;
+}
+
 #if WITH_EDITOR
 FText UMovieGraphObjectIdNode::GetNodeTitle(const bool bGetDescriptive) const
 {
diff --git a/Engine/Plugins/MovieScene/MoviePipelineMaskRenderPass/Source/MoviePipelineMaskRenderPass/Public/Graph/MovieGraphObjectIdNode.h b/Engine/Plugins/MovieScene/MoviePipelineMaskRenderPass/Source/MoviePipelineMaskRenderPass/Public/Graph/MovieGraphObjectIdNode.h
--- a/Engine/Plugins/MovieScene/MoviePipelineMaskRenderPass/Source/MoviePipelineMaskRenderPass/Public/Graph/MovieGraphObjectIdNode.h
+++ b/Engine/Plugins/MovieScene/MoviePipelineMaskRenderPass/Source/MoviePipelineMaskRenderPass/Public/Graph/MovieGraphObjectIdNode.h
@@ -25,11 +25,12 @@ public:
 	UE_API virtual EViewModeIndex GetViewModeIndex() const override;
 	UE_API virtual bool GetAllowsShowFlagsCustomization() const override;
 	UE_API virtual bool GetAllowsCompositing() const override;
-	virtual int32 GetNumSpatialSamples() const override { return SpatialSampleCount; }
-	virtual int32 GetNumSpatialSamplesDuringWarmUp() const  override { return 0; }
+	UE_API virtual bool GetForceLosslessCompression() const override;
+	UE_API virtual int32 GetNumSpatialSamples() const override { return SpatialSampleCount; }
+	UE_API virtual int32 GetNumSpatialSamplesDuringWarmUp() const  override { return 0; }
 	// The ObjectID pass doesn't support any anti-aliasing so we force it to None so that their spatial samples turn into MRG-based camera jitter.
-	virtual bool GetOverrideAntiAliasing() const override { return true; }
-	virtual EAntiAliasingMethod GetAntiAliasingMethod() const override { return EAntiAliasingMethod::AAM_None; }
+	UE_API virtual bool GetOverrideAntiAliasing() const override { return true; }
+	UE_API virtual EAntiAliasingMethod GetAntiAliasingMethod() const override { return EAntiAliasingMethod::AAM_None; }
 	// ~UMovieGraphImagePassBaseNode Interface
 
 	// UMovieGraphNode Interface
diff --git a/Engine/Plugins/MovieScene/MovieRenderPipeline/Source/MovieRenderPipelineCore/Public/Graph/MovieGraphDataTypes.h b/Engine/Plugins/MovieScene/MovieRenderPipeline/Source/MovieRenderPipelineCore/Public/Graph/MovieGraphDataTypes.h
--- a/Engine/Plugins/MovieScene/MovieRenderPipeline/Source/MovieRenderPipelineCore/Public/Graph/MovieGraphDataTypes.h
+++ b/Engine/Plugins/MovieScene/MovieRenderPipeline/Source/MovieRenderPipelineCore/Public/Graph/MovieGraphDataTypes.h
@@ -406,6 +406,7 @@ namespace UE::MovieGraph
 			, CompositingSortOrder(0)
 			, RenderLayerIndex(0)
 			, bAllowOCIO(true)
+			, bForceLosslessCompression(false)
 		{}
 
 		virtual TSharedRef<FMovieGraphSampleState> Copy() const
@@ -483,6 +484,9 @@ namespace UE::MovieGraph
 		/** Allow OpenColorIO transform to be used on this render. */
 		bool bAllowOCIO;
 
+		/** Whether this sample must use lossless compression when written to a format that supports it. */
+		bool bForceLosslessCompression;
+
 		/** Render scene capture source used for tracking the output color space (without OpenColorIO). */
 		ESceneCaptureSource SceneCaptureSource;
 
diff --git a/Engine/Plugins/MovieScene/MovieRenderPipeline/Source/MovieRenderPipelineEditor/Private/Customizations/PostProcessPassCustomization.h b/Engine/Plugins/MovieScene/MovieRenderPipeline/Source/MovieRenderPipelineEditor/Private/Customizations/PostProcessPassCustomization.h
new file mode 100644
--- /dev/null
+++ b/Engine/Plugins/MovieScene/MovieRenderPipeline/Source/MovieRenderPipelineEditor/Private/Customizations/PostProcessPassCustomization.h
@@ -0,0 +1,64 @@
+﻿// Copyright Epic Games, Inc. All Rights Reserved.
+#pragma once
+
+#include "DetailWidgetRow.h"
+#include "IDetailChildrenBuilder.h"
+#include "IPropertyTypeCustomization.h"
+#include "PropertyHandle.h"
+
+#include "MoviePipelineDeferredPasses.h"
+
+#define LOCTEXT_NAMESPACE "SMoviePipelineEditor"
+
+/* Customizes how post process passes are displayed (for MRQ, not MRG). */
+class FPostProcessPassCustomization_MRQ : public IPropertyTypeCustomization
+{
+public:
+	static TSharedRef<IPropertyTypeCustomization> MakeInstance()
+	{
+		return MakeShared<FPostProcessPassCustomization_MRQ>();
+	}
+
+protected:
+	//~ Begin IPropertyTypeCustomization interface
+	virtual void CustomizeHeader(TSharedRef<IPropertyHandle> PropertyHandle, FDetailWidgetRow& HeaderRow, IPropertyTypeCustomizationUtils& CustomizationUtils) override
+	{
+		// Include the default header row so the array index "groups" show up (eg, "Index [0]").
+		HeaderRow
+		.NameContent()
+		[
+			PropertyHandle->CreatePropertyNameWidget()
+		]
+		.ValueContent()
+		[
+			PropertyHandle->CreatePropertyValueWidget()
+		];
+	}
+	
+	virtual void CustomizeChildren(TSharedRef<IPropertyHandle> InStructPropertyHandle, IDetailChildrenBuilder& StructBuilder, IPropertyTypeCustomizationUtils& StructCustomizationUtils) override
+	{
+		static const TArray<FName> HiddenPropertyNames =
+		{
+			// The "Use Lossless Compression" property is not implemented for MRQ, only MRG
+			GET_MEMBER_NAME_CHECKED(FMoviePipelinePostProcessPass, bUseLosslessCompression)
+		};
+		
+		// Add all struct properties except the hidden ones
+		uint32 NumChildren;
+		if (InStructPropertyHandle->GetNumChildren(NumChildren) == FPropertyAccess::Success)
+		{
+			for (uint32 Index = 0; Index < NumChildren; Index++)
+			{
+				TSharedPtr<IPropertyHandle> ChildProperty = InStructPropertyHandle->GetChildHandle(Index);
+
+				if (!HiddenPropertyNames.Contains(ChildProperty->GetProperty()->GetFName()))
+				{
+					StructBuilder.AddProperty(ChildProperty.ToSharedRef());
+				}
+			}
+		}
+	}
+	//~ End IPropertyTypeCustomization interface
+};
+
+#undef LOCTEXT_NAMESPACE
\ No newline at end of file
diff --git a/Engine/Plugins/MovieScene/MovieRenderPipeline/Source/MovieRenderPipelineEditor/Private/Widgets/SMoviePipelineConfigEditor.cpp b/Engine/Plugins/MovieScene/MovieRenderPipeline/Source/MovieRenderPipelineEditor/Private/Widgets/SMoviePipelineConfigEditor.cpp
--- a/Engine/Plugins/MovieScene/MovieRenderPipeline/Source/MovieRenderPipelineEditor/Private/Widgets/SMoviePipelineConfigEditor.cpp
+++ b/Engine/Plugins/MovieScene/MovieRenderPipeline/Source/MovieRenderPipelineEditor/Private/Widgets/SMoviePipelineConfigEditor.cpp
@@ -3,6 +3,7 @@
 #include "Widgets/SMoviePipelineConfigEditor.h"
 
 #include "MoviePipelineConsoleVariableSetting.h"
+#include "MoviePipelineDeferredPasses.h"
 #include "MoviePipelineOutputSetting.h"
 #include "MoviePipelineSetting.h"
 #include "MoviePipelineQueue.h"
@@ -37,14 +38,15 @@
 #include "ScopedTransaction.h"
 #include "PropertyEditorModule.h"
 #include "IDetailsView.h"
-#include "FrameNumberDetailsCustomization.h"
 #include "IDetailCustomization.h"
 #include "Editor.h"
 
 #include "DetailLayoutBuilder.h"
 #include "DetailWidgetRow.h"
 
+#include "FrameNumberDetailsCustomization.h"
 #include "Customizations/ConsoleVariableSettingCustomization.h"
+#include "Customizations/PostProcessPassCustomization.h"
 
 #define LOCTEXT_NAMESPACE "SMoviePipelineEditor"
 
@@ -409,6 +411,7 @@ void SMoviePipelineConfigEditor::UpdateDetails()
 		{
 			TSharedRef<IDetailsView> Details = PropertyEditorModule.CreateDetailView(DetailsViewArgs);
 			Details->RegisterInstancedCustomPropertyTypeLayout("FrameNumber", FOnGetPropertyTypeCustomizationInstance::CreateStatic(&FFrameNumberDetailsCustomization::MakeInstance, NumericTypeInterface));
+			Details->RegisterInstancedCustomPropertyTypeLayout(FMoviePipelinePostProcessPass::StaticStruct()->GetFName(), FOnGetPropertyTypeCustomizationInstance::CreateStatic(FPostProcessPassCustomization_MRQ::MakeInstance));
 			Details->RegisterInstancedCustomPropertyLayout(UMoviePipelineOutputSetting::StaticClass(), FOnGetDetailCustomizationInstance::CreateStatic(FOutputFormatDetailsCustomization::MakeInstance, SharedThis(this)));
 			Details->RegisterInstancedCustomPropertyLayout(UMoviePipelineConsoleVariableSetting::StaticClass(), FOnGetDetailCustomizationInstance::CreateStatic(FConsoleVariablesSettingDetailsCustomization::MakeInstance));
 
diff --git a/Engine/Plugins/MovieScene/MovieRenderPipeline/Source/MovieRenderPipelineRenderPasses/Private/Graph/Nodes/MovieGraphDeferredPassNode.cpp b/Engine/Plugins/MovieScene/MovieRenderPipeline/Source/MovieRenderPipelineRenderPasses/Private/Graph/Nodes/MovieGraphDeferredPassNode.cpp
--- a/Engine/Plugins/MovieScene/MovieRenderPipeline/Source/MovieRenderPipelineRenderPasses/Private/Graph/Nodes/MovieGraphDeferredPassNode.cpp
+++ b/Engine/Plugins/MovieScene/MovieRenderPipeline/Source/MovieRenderPipelineRenderPasses/Private/Graph/Nodes/MovieGraphDeferredPassNode.cpp
@@ -37,6 +37,7 @@ UMovieGraphDeferredRenderPassNode::UMovieGraphDeferredRenderPassNode()
 		NewPass.Material = TSoftObjectPtr<UMaterialInterface>(FSoftObjectPath(MaterialPath));
 		NewPass.bEnabled = false;
 		NewPass.bHighPrecisionOutput = MaterialPath.Equals(DefaultDepthAsset);
+		NewPass.bUseLosslessCompression = true;
 	}
 }
 
diff --git a/Engine/Plugins/MovieScene/MovieRenderPipeline/Source/MovieRenderPipelineRenderPasses/Private/Graph/Nodes/MovieGraphPathTracerPassNode.cpp b/Engine/Plugins/MovieScene/MovieRenderPipeline/Source/MovieRenderPipelineRenderPasses/Private/Graph/Nodes/MovieGraphPathTracerPassNode.cpp
--- a/Engine/Plugins/MovieScene/MovieRenderPipeline/Source/MovieRenderPipelineRenderPasses/Private/Graph/Nodes/MovieGraphPathTracerPassNode.cpp
+++ b/Engine/Plugins/MovieScene/MovieRenderPipeline/Source/MovieRenderPipelineRenderPasses/Private/Graph/Nodes/MovieGraphPathTracerPassNode.cpp
@@ -42,6 +42,7 @@ UMovieGraphPathTracerRenderPassNode::UMovieGraphPathTracerRenderPassNode()
 		NewPass.Material = TSoftObjectPtr<UMaterialInterface>(FSoftObjectPath(MaterialPath));
 		NewPass.bEnabled = false;
 		NewPass.bHighPrecisionOutput = MaterialPath.Equals(DefaultDepthAsset);
+		NewPass.bUseLosslessCompression = true;
 	}
 }
 
diff --git a/Engine/Plugins/MovieScene/MovieRenderPipeline/Source/MovieRenderPipelineRenderPasses/Private/Graph/Renderers/MovieGraphDeferredPanoramicPass.cpp b/Engine/Plugins/MovieScene/MovieRenderPipeline/Source/MovieRenderPipelineRenderPasses/Private/Graph/Renderers/MovieGraphDeferredPanoramicPass.cpp
--- a/Engine/Plugins/MovieScene/MovieRenderPipeline/Source/MovieRenderPipelineRenderPasses/Private/Graph/Renderers/MovieGraphDeferredPanoramicPass.cpp
+++ b/Engine/Plugins/MovieScene/MovieRenderPipeline/Source/MovieRenderPipelineRenderPasses/Private/Graph/Renderers/MovieGraphDeferredPanoramicPass.cpp
@@ -682,6 +682,7 @@ void FMovieGraphDeferredPanoramicPass::Render(const FMovieGraphTraversalContext&
 					SampleState.OverscanFraction = 0.f;
 					SampleState.CropRectangle = FIntRect(0, 0, DesiredOutputResolution.X, DesiredOutputResolution.Y); // ToDo: Output resolution will be forced to a 2:1 ratio but this currently respects what the user put in.
 					SampleState.bAllowOCIO = ParentNodeThisFrame->GetAllowOCIO();
+					SampleState.bForceLosslessCompression = ParentNodeThisFrame->GetForceLosslessCompression();
 					SampleState.bAllowsCompositing = ParentNodeThisFrame->GetAllowsCompositing();
 					SampleState.SceneCaptureSource = SceneCaptureSource;
 					SampleState.CompositingSortOrder = 10;
diff --git a/Engine/Plugins/MovieScene/MovieRenderPipeline/Source/MovieRenderPipelineRenderPasses/Private/Graph/Renderers/MovieGraphDeferredPass.cpp b/Engine/Plugins/MovieScene/MovieRenderPipeline/Source/MovieRenderPipelineRenderPasses/Private/Graph/Renderers/MovieGraphDeferredPass.cpp
--- a/Engine/Plugins/MovieScene/MovieRenderPipeline/Source/MovieRenderPipelineRenderPasses/Private/Graph/Renderers/MovieGraphDeferredPass.cpp
+++ b/Engine/Plugins/MovieScene/MovieRenderPipeline/Source/MovieRenderPipelineRenderPasses/Private/Graph/Renderers/MovieGraphDeferredPass.cpp
@@ -465,6 +465,7 @@ void FMovieGraphDeferredPass::Render(const FMovieGraphTraversalContext& InFrameT
 					SampleState.OverscanFraction = CameraInfo.ViewInfo.GetOverscan();
 					SampleState.CropRectangle = AccumulatorResolutionCropRect;
 					SampleState.bAllowOCIO = ParentNodeThisFrame->GetAllowOCIO();
+					SampleState.bForceLosslessCompression = ParentNodeThisFrame->GetForceLosslessCompression();
 					SampleState.bAllowsCompositing = ParentNodeThisFrame->GetAllowsCompositing();
 					SampleState.SceneCaptureSource = SceneCaptureSource;
 					SampleState.CompositingSortOrder = 10;
@@ -518,7 +519,10 @@ void FMovieGraphDeferredPass::Render(const FMovieGraphTraversalContext& InFrameT
 						// nor should they have OCIO applied (as they're going to be data buffers like depth).
 						PassSampleState.bAllowsCompositing = false;
 						PassSampleState.bAllowOCIO = false;
-				
+
+						// Some Post Process materials may need lossless compression in order to preserve the data properly.
+						PassSampleState.bForceLosslessCompression = PostProcessPass.bUseLosslessCompression;
+
 						// Give a lower priority to materials so they show up after the main pass in multi-layer exrs.
 						PassSampleState.CompositingSortOrder = SampleState.CompositingSortOrder + 1;
 						BufferPipe->AddEndpoint(MakeForwardingEndpoint(PassSampleState, InTimeData));
diff --git a/Engine/Plugins/MovieScene/MovieRenderPipeline/Source/MovieRenderPipelineRenderPasses/Private/MovieGraphImageSequenceOutputNode.cpp b/Engine/Plugins/MovieScene/MovieRenderPipeline/Source/MovieRenderPipelineRenderPasses/Private/MovieGraphImageSequenceOutputNode.cpp
--- a/Engine/Plugins/MovieScene/MovieRenderPipeline/Source/MovieRenderPipelineRenderPasses/Private/MovieGraphImageSequenceOutputNode.cpp
+++ b/Engine/Plugins/MovieScene/MovieRenderPipeline/Source/MovieRenderPipelineRenderPasses/Private/MovieGraphImageSequenceOutputNode.cpp
@@ -23,6 +23,21 @@
 
 #include UE_INLINE_GENERATED_CPP_BY_NAME(MovieGraphImageSequenceOutputNode)
 
+/*
+ * Log a message once per render.
+ *
+ * Note that the arguments to this are the same as for UE_LOG, but also takes an additional `Identifier`. The identifier is what uniquely identifies
+ * this specific message within the current SCOPE (identifiers are NOT unique across scopes, so the same log message will be logged more than once per
+ * render if it appears in different scopes).
+ */
+#define MRG_LOG_ONCE_PER_RENDER(Identifier, LogCategory, LogVerbosity, LogFormat, ...)							\
+	static UWorld* PREPROCESSOR_JOIN(LastRenderWorld_, Identifier) = nullptr;									\
+	if (PREPROCESSOR_JOIN(LastRenderWorld_, Identifier) != MoviePipeline::FindCurrentWorld())					\
+	{																											\
+		PREPROCESSOR_JOIN(LastRenderWorld_, Identifier) = MoviePipeline::FindCurrentWorld();					\
+		UE_LOG(LogCategory, LogVerbosity, LogFormat, ##__VA_ARGS__);											\
+	}																											\
+
 UMovieGraphImageSequenceOutputNode::UMovieGraphImageSequenceOutputNode()
 {
 	ImageWriteQueue = &FModuleManager::Get().LoadModuleChecked<IImageWriteQueueModule>("ImageWriteQueue").GetWriteQueue();
@@ -426,8 +441,11 @@ void UMovieGraphImageSequenceOutputNode_EXR::OnReceiveImageDataImpl(UMovieGraphP
 		}
 
 		UE::MovieGraph::FMovieGraphSampleState* Payload = RenderData.Value->GetPayload<UE::MovieGraph::FMovieGraphSampleState>();
+
+		// The pass may be need to be written out with lossless compression (eg, Object ID forces this on).
+		const EEXRCompressionFormat CompressionType = Payload->bForceLosslessCompression ? ParentNode->LosslessCompression : ParentNode->Compression;
 		
-		TUniquePtr<FEXRImageWriteTask> ImageWriteTask = CreateImageWriteTask(FileName, ParentNode->Compression);
+		TUniquePtr<FEXRImageWriteTask> ImageWriteTask = CreateImageWriteTask(FileName, CompressionType);
 		
 		PrepareTaskGlobalMetadata(*ImageWriteTask, InRawFrameData, ResolvedFormatArgs.FileMetadata);
 
@@ -503,6 +521,12 @@ void UMovieGraphImageSequenceOutputNode_EXR::OnReceiveImageDataImpl(UMovieGraphP
 
 void UMovieGraphImageSequenceOutputNode_MultiLayerEXR::OnReceiveImageDataImpl(UMovieGraphPipeline* InPipeline, UE::MovieGraph::FMovieGraphOutputMergerFrame* InRawFrameData, const TSet<FMovieGraphRenderDataIdentifier>& InMask)
 {
+	auto GetPayloadCompressionType = [](
+		const UE::MovieGraph::FMovieGraphSampleState* InPayload, const UMovieGraphImageSequenceOutputNode_MultiLayerEXR* InParentNode) -> EEXRCompressionFormat
+	{
+		return InPayload->bForceLosslessCompression ? InParentNode->LosslessCompression : InParentNode->Compression;
+	};
+	
 	check(InRawFrameData);
 
 	constexpr bool bIncludeCDOs = false;
@@ -522,7 +546,33 @@ void UMovieGraphImageSequenceOutputNode_MultiLayerEXR::OnReceiveImageDataImpl(UM
 		const FString& Filename = RenderConfigForFilename.Key;
 		FEXROutputConfigForFilename& RenderConfig = RenderConfigForFilename.Value;
 		
-		TUniquePtr<FEXRImageWriteTask> MultiLayerImageTask = CreateImageWriteTask(Filename, ParentNode->Compression, ParentNode->bMultipart);
+		// If all layers are not using the same compression type, force multipart on.
+		bool bUseMultipart = ParentNode->bMultipart;
+		if (!bUseMultipart)
+		{
+			EEXRCompressionFormat LastCompressionType = EEXRCompressionFormat::Max;
+			
+			for (const FMovieGraphRenderDataIdentifier& RenderID : RenderConfig.RenderIDs)
+			{
+				const TUniquePtr<FImagePixelData>& ImageData = InRawFrameData->ImageOutputData[RenderID];
+				checkf(ImageData.IsValid(), TEXT("Unexpected empty image data: incorrectly moved or its production failed?"));
+
+				const UE::MovieGraph::FMovieGraphSampleState* Payload = ImageData->GetPayload<UE::MovieGraph::FMovieGraphSampleState>();
+				const EEXRCompressionFormat CompressionType = GetPayloadCompressionType(Payload, ParentNode);
+				if (LastCompressionType == EEXRCompressionFormat::Max)
+				{
+					LastCompressionType = CompressionType;
+				}
+				else if (CompressionType != LastCompressionType)
+				{
+					MRG_LOG_ONCE_PER_RENDER(ForceMultiPart, LogMovieRenderPipeline, Warning, TEXT("An EXR is being generated that containins layers with differing compression types (can happen especially when generating Object IDs or using PPMs). Forcing EXR multipart ON."));
+					bUseMultipart = true;
+					break;
+				}
+			}
+		}
+		
+		TUniquePtr<FEXRImageWriteTask> MultiLayerImageTask = CreateImageWriteTask(Filename, ParentNode->Compression, bUseMultipart);
 		PrepareTaskGlobalMetadata(*MultiLayerImageTask, InRawFrameData, RenderConfig.ResolveArgs.FileMetadata);
 
 		// Keep track of the lowest render layer index found among the layers that are included. This will be used as the index provided to the output
@@ -544,6 +594,13 @@ void UMovieGraphImageSequenceOutputNode_MultiLayerEXR::OnReceiveImageDataImpl(UM
 
 			LowestRenderLayerIndex = FMath::Min(LowestRenderLayerIndex, Payload->RenderLayerIndex);
 
+			// If using multipart, specify the compression type for each layer.
+			if (bUseMultipart)
+			{
+				const EEXRCompressionFormat LayerCompressionType = GetPayloadCompressionType(Payload, ParentNode);
+				MultiLayerImageTask->CompressionByLayer.Add(LayerCompressionType);
+			}
+
 			// The first layer that doesn't have an explicitly-specified name will have an empty layer name in the EXR -- this is the "primary"/RGBA
 			// layer. Otherwise, the layer name will be procedurally generated if the payload didn't specify a layer name override. Having a "primary"
 			// layer in the EXR expands compatibility with a number of applications that read EXRs.
diff --git a/Engine/Plugins/MovieScene/MovieRenderPipeline/Source/MovieRenderPipelineRenderPasses/Private/MoviePipelineEXROutput.cpp b/Engine/Plugins/MovieScene/MovieRenderPipeline/Source/MovieRenderPipelineRenderPasses/Private/MoviePipelineEXROutput.cpp
--- a/Engine/Plugins/MovieScene/MovieRenderPipeline/Source/MovieRenderPipelineRenderPasses/Private/MoviePipelineEXROutput.cpp
+++ b/Engine/Plugins/MovieScene/MovieRenderPipeline/Source/MovieRenderPipelineRenderPasses/Private/MoviePipelineEXROutput.cpp
@@ -126,8 +126,26 @@ void FEXRImageWriteTask::OnAbandoned()
 	}
 }
 
-bool FEXRImageWriteTask::GenerateFilePartsForResolution(const FIntPoint& Resolution, const TArray<FImagePixelData*>& InLayers, const FString& InPartName, FEXRPart& OutPart, TArray<TUniquePtr<FImagePixelData>>& OutQuantizedLayers)
+bool FEXRImageWriteTask::GenerateFilePartsForResolution(const FIntPoint& Resolution, const TArray<FImagePixelData*>& InLayers, const EEXRCompressionFormat CompressionType, const FString& InPartName, FEXRPart& OutPart, TArray<TUniquePtr<FImagePixelData>>& OutQuantizedLayers)
 {
+	// Generate the header that's either used for the entire file, or per-layer in multi-part EXRs.
+	auto GenerateHeader = [this, &InPartName]
+		(const IMATH_NAMESPACE::Box2i& InDisplayWindow, const IMATH_NAMESPACE::Box2i& InDataWindow, const Imf::Compression CompressionType) -> Imf::Header
+	{
+		Imf::Header Header(InDisplayWindow, InDataWindow, 1, Imath::V2f(0, 0), 1, Imf::LineOrder::INCREASING_Y, CompressionType);
+
+		if (bMultipart)
+		{
+			Header.setName(TCHAR_TO_ANSI(*InPartName));
+			Header.setType("scanlineimage");
+		}
+
+		// Add all of the `FileMetadata` metadata into the header.
+        AddFileMetadata(Header);
+
+		return Header;
+	};
+	
 	bool bSuccess = false;
 
 	OutPart.Scanlines = bPadToDataWindowSize ? Height : Resolution.Y;
@@ -159,16 +177,7 @@ bool FEXRImageWriteTask::GenerateFilePartsForResolution(const FIntPoint& Resolut
 	IMATH_NAMESPACE::Box2i DisplayWindow = IMATH_NAMESPACE::Box2i(IMATH_NAMESPACE::V2i(0, 0),IMATH_NAMESPACE::V2i(DisplayWindowRes.X - 1, DisplayWindowRes.Y - 1));
 
 	static_assert(static_cast<uint8>(EEXRCompressionFormat::Max) == Imf::Compression::NUM_COMPRESSION_METHODS);
-	Imf::Compression FileCompression = static_cast<Imf::Compression>(Compression);
-
-	OutPart.Header = Imf::Header(DisplayWindow, DataWindow, 1, Imath::V2f(0, 0), 1, Imf::LineOrder::INCREASING_Y, FileCompression);
-
-	// Multi-part headers need unique names and need their type specified
-	if (bMultipart)
-	{
-		OutPart.Header.setName(TCHAR_TO_ANSI(*InPartName));
-		OutPart.Header.setType("scanlineimage");
-	}
+	Imf::Compression FileCompression = static_cast<Imf::Compression>(CompressionType);
 	
 	// If using lossy compression, specify the compression level in the header per exr spec.
 	if (FileCompression == Imf::Compression::DWAA_COMPRESSION ||
@@ -177,8 +186,7 @@ bool FEXRImageWriteTask::GenerateFilePartsForResolution(const FIntPoint& Resolut
 		FileMetadata.Add("dwaCompressionLevel", CompressionLevel);
 	}
 
-	// Insert our key-value pair metadata (if any, can be an arbitrary set of key/value pairs)
-	AddFileMetadata(OutPart.Header);
+	OutPart.Header = GenerateHeader(DisplayWindow, DataWindow, FileCompression);
 
 	if (ColorSpaceChromaticities.Num() > 0)
 	{
@@ -310,11 +318,17 @@ bool FEXRImageWriteTask::WriteToDisk()
 		
 		if (bMultipart)
 		{
-			for (TUniquePtr<FImagePixelData>& Layer : Layers)
+			for (int32 LayerIndex = 0; LayerIndex < Layers.Num(); ++LayerIndex)
 			{
 				FEXRPart Part;
+				TUniquePtr<FImagePixelData>& Layer = Layers[LayerIndex];
 				FString PartName = LayerNames.Contains(Layer.Get()) ? LayerNames[Layer.Get()] : TEXT("FinalImage");
-				if (!GenerateFilePartsForResolution(Layer->GetSize(), { Layer.Get() }, PartName, Part, QuantizedLayers))
+
+				// Override the compression level for the layer if requested.
+				const bool bUsingPerLayerCompressionType = CompressionByLayer.Num() == Layers.Num();
+				const EEXRCompressionFormat LayerCompressionType = bUsingPerLayerCompressionType ? CompressionByLayer[LayerIndex] : Compression;
+				
+				if (!GenerateFilePartsForResolution(Layer->GetSize(), { Layer.Get() }, LayerCompressionType, PartName, Part, QuantizedLayers))
 				{
 					UE_LOG(LogMovieRenderPipelineIO, Error, TEXT("There was an error writing layer '%s' to the EXR file. The pixel format may not be compatible with this image type, or there was a resolution mismatch."), *LayerNames[Layer.Get()]);
 					bSuccess = false;
@@ -329,7 +343,8 @@ bool FEXRImageWriteTask::WriteToDisk()
 			FEXRPart Part;
 			TArray<FImagePixelData*> LayerRawPtrs;
 			Algo::Transform(Layers, LayerRawPtrs, [](const TUniquePtr<FImagePixelData>& Layer) { return Layer.Get(); });
-			if (!GenerateFilePartsForResolution(FIntPoint(Width, Height), LayerRawPtrs, TEXT("FinalImage"), Part, QuantizedLayers))
+			
+			if (!GenerateFilePartsForResolution(FIntPoint(Width, Height), LayerRawPtrs, Compression, TEXT("FinalImage"), Part, QuantizedLayers))
 			{
 				bSuccess = false;
 			}
diff --git a/Engine/Plugins/MovieScene/MovieRenderPipeline/Source/MovieRenderPipelineRenderPasses/Public/Graph/Nodes/MovieGraphImagePassBaseNode.h b/Engine/Plugins/MovieScene/MovieRenderPipeline/Source/MovieRenderPipelineRenderPasses/Public/Graph/Nodes/MovieGraphImagePassBaseNode.h
--- a/Engine/Plugins/MovieScene/MovieRenderPipeline/Source/MovieRenderPipelineRenderPasses/Public/Graph/Nodes/MovieGraphImagePassBaseNode.h
+++ b/Engine/Plugins/MovieScene/MovieRenderPipeline/Source/MovieRenderPipelineRenderPasses/Public/Graph/Nodes/MovieGraphImagePassBaseNode.h
@@ -75,6 +75,9 @@ public:
 
 	/** Whether this pass allows other passes to be composited on it. */
 	virtual bool GetAllowsCompositing() const { return true; }
+
+	/** Whether lossless compression must be forced on for this node. */
+	virtual bool GetForceLosslessCompression() const { return false; }
 	
 	/** Are we using high resolution tiling? */
 	virtual bool GetEnableHighResolutionTiling() const { return false; }
@@ -90,6 +93,7 @@ public:
 	
 	/** If using high resolution tiling, do we keep a unique scene view history for each tile? */
 	virtual bool GetEnableHistoryPerTile() const { return false; }
+
 protected:
 	// Note: Since *individual* show flags are overridden instead of the entire ShowFlags property, manually set to
 	// overridden so the traversal picks the changes up (otherwise they will be ignored).
diff --git a/Engine/Plugins/MovieScene/MovieRenderPipeline/Source/MovieRenderPipelineRenderPasses/Public/MovieGraphImageSequenceOutputNode.h b/Engine/Plugins/MovieScene/MovieRenderPipeline/Source/MovieRenderPipelineRenderPasses/Public/MovieGraphImageSequenceOutputNode.h
--- a/Engine/Plugins/MovieScene/MovieRenderPipeline/Source/MovieRenderPipelineRenderPasses/Public/MovieGraphImageSequenceOutputNode.h
+++ b/Engine/Plugins/MovieScene/MovieRenderPipeline/Source/MovieRenderPipelineRenderPasses/Public/MovieGraphImageSequenceOutputNode.h
@@ -105,6 +105,7 @@ public:
 		OutputFormat = EImageFormat::EXR;
 		bQuantizeTo8Bit = false;
 		Compression = EEXRCompressionFormat::PIZ;
+		LosslessCompression = EEXRCompressionFormat::PIZ;
 	}
 
 #if WITH_UNREALEXR
@@ -144,12 +145,23 @@ public:
 	UPROPERTY(EditAnywhere, BlueprintReadWrite, Category = Overrides, meta = (InlineEditConditionToggle))
 	uint8 bOverride_Compression : 1;
 	
+	UPROPERTY(EditAnywhere, BlueprintReadWrite, Category = Overrides, meta = (InlineEditConditionToggle))
+	uint8 bOverride_LosslessCompression : 1;
+	
 	/**
-	 * Which compression method should the resulting EXR file be compressed with.
+	 * The type of compression the EXR file should be compressed with. This compression method will be applied to all layers unless the data in
+	 * a layer has been flagged as needing lossless compression. See "Lossless Compression" for more details.
 	 */
 	UPROPERTY(EditAnywhere, BlueprintReadWrite, meta=(EditCondition="bOverride_Compression"), Category = "EXR")
 	EEXRCompressionFormat Compression;
 
+	/**
+	 * The type of lossless compression that should be used when a render pass is either forced to use lossless (eg, Object ID), or when the pass
+	 * has been flagged to need lossless (eg, Post Process Material passes can opt-in to lossless).
+	 */
+	UPROPERTY(EditAnywhere, BlueprintReadWrite, meta=(EditCondition="bOverride_LosslessCompression", ValidEnumValues="None,RLE,ZIPS,ZIP,PIZ"), Category = "EXR")
+	EEXRCompressionFormat LosslessCompression;
+
 protected:
 
 #if WITH_UNREALEXR
@@ -196,7 +208,6 @@ public:
 	{
 		// Multi-layer default excludes {layer_name}.
 		FileNameFormat = TEXT("{sequence_name}.{frame_number}");
-		bMultipart = false;
 	}
 
 	virtual EMovieGraphBranchRestriction GetBranchRestriction() const override
@@ -281,9 +292,12 @@ public:
 	 * Indicates whether the exr file should be written as a multi-part exr file, which supports having different image types and resolutions for each layer.
 	 * Multi-part EXRs are a feature of EXR 2.0 and may not be supported by all software. If the console variable 'MoviePipeline.PadLayersForMultiPartEXR' is enabled,
 	 * then all parts of the multi-part EXR will be padded to match the data window of the largest layer, as some software does not support different data window sizes.
+	 *
+	 * Note that multi-part will be forced on if layers are using more than one compression type (this can happen if rendering a beauty pass using lossy
+	 * compression, and Object IDs using lossless compression, for example).
 	 */
 	UPROPERTY(EditAnywhere, BlueprintReadWrite, Category = "EXR", meta=(EditCondition="bOverride_bMultiPart"))
-	bool bMultipart = false;
+	bool bMultipart = true;
 };
 
 /**
diff --git a/Engine/Plugins/MovieScene/MovieRenderPipeline/Source/MovieRenderPipelineRenderPasses/Public/MoviePipelineDeferredPasses.h b/Engine/Plugins/MovieScene/MovieRenderPipeline/Source/MovieRenderPipelineRenderPasses/Public/MoviePipelineDeferredPasses.h
--- a/Engine/Plugins/MovieScene/MovieRenderPipeline/Source/MovieRenderPipelineRenderPasses/Public/MoviePipelineDeferredPasses.h
+++ b/Engine/Plugins/MovieScene/MovieRenderPipeline/Source/MovieRenderPipelineRenderPasses/Public/MoviePipelineDeferredPasses.h
@@ -43,6 +43,13 @@ public:
 	/** Request output to be 32-bit, usually for data exports. Note that scene color precision is still defined by r.SceneColorFormat.*/
 	UPROPERTY(EditAnywhere, BlueprintReadWrite, Category = "Settings", DisplayName="Use High Precision (32-bit) Output")
 	bool bHighPrecisionOutput = false;
+
+	/**
+	 * Whether this pass should use lossless compression in formats that support it. Should usually be turned on for passes that need a highly
+	 * accurate representation of the data (eg, depth).
+	 */
+	UPROPERTY(EditAnywhere, BlueprintReadWrite, Category = "Settings")
+	bool bUseLosslessCompression = false;
 };
 
 UCLASS(MinimalAPI, BlueprintType)
diff --git a/Engine/Plugins/MovieScene/MovieRenderPipeline/Source/MovieRenderPipelineRenderPasses/Public/MoviePipelineEXROutput.h b/Engine/Plugins/MovieScene/MovieRenderPipeline/Source/MovieRenderPipelineRenderPasses/Public/MoviePipelineEXROutput.h
--- a/Engine/Plugins/MovieScene/MovieRenderPipeline/Source/MovieRenderPipelineRenderPasses/Public/MoviePipelineEXROutput.h
+++ b/Engine/Plugins/MovieScene/MovieRenderPipeline/Source/MovieRenderPipelineRenderPasses/Public/MoviePipelineEXROutput.h
@@ -87,9 +87,18 @@ public:
 	/** True if this task is allowed to overwrite an existing file, false otherwise. */
 	bool bOverwriteFile;
 	
-	/** Compression method used for the resulting EXR files. */
+	/**
+	 * Compression method used for the resulting EXR files. Applies to all layers, but per-layer compression can be specified in
+	 * `CompressionByLayer` if needed.
+	 */
 	EEXRCompressionFormat Compression;
 
+	/**
+	 * Controls the compression per layer. Only used if `bMultiPart` is true. Takes precedence over `Compression`, but is ignored if this array is
+	 * either empty, or its size does not match the number of layers given in `Layers`.
+	 */
+	TArray<EEXRCompressionFormat> CompressionByLayer;
+
 	/** When using a lossy compression format, what is the base-error (CompressionLevel/100000) */
 	int32 CompressionLevel;
 
@@ -154,7 +163,7 @@ private:
 	};
 
 	/** Generates EXR file parts (header, framebuffer, etc) for the specified layers at the specified resolution */
-	bool GenerateFilePartsForResolution(const FIntPoint& Resolution, const TArray<FImagePixelData*>& InLayers, const FString& InPartName, FEXRPart& OutPart, TArray<TUniquePtr<FImagePixelData>>& OutQuantizedLayers);
+	bool GenerateFilePartsForResolution(const FIntPoint& Resolution, const TArray<FImagePixelData*>& InLayers, const EEXRCompressionFormat CompressionType, const FString& InPartName, FEXRPart& OutPart, TArray<TUniquePtr<FImagePixelData>>& OutQuantizedLayers);
 	
 	/**
 	 * Run the task, attempting to write out the raw data using the currently specified parameters
@@ -228,5 +237,4 @@ public:
 	 */
 	UPROPERTY(EditAnywhere, BlueprintReadWrite, Category = "EXR")
 	bool bMultipart;
-
 };
