MoviePipeline: Added a way to access evaluated variables from the graph's burn-in UMG widget.

GetVariableValueByName() and GetVariableStringValueByName() can now be called on the evaluated graph within the burn-in UMG graph to fetch fully evaluated variable values.

diff --git a/Engine/Plugins/MovieScene/MovieRenderPipeline/Source/MovieRenderPipelineCore/Private/Graph/MovieGraphConfig.cpp b/Engine/Plugins/MovieScene/MovieRenderPipeline/Source/MovieRenderPipelineCore/Private/Graph/MovieGraphConfig.cpp
--- a/Engine/Plugins/MovieScene/MovieRenderPipeline/Source/MovieRenderPipelineCore/Private/Graph/MovieGraphConfig.cpp
+++ b/Engine/Plugins/MovieScene/MovieRenderPipeline/Source/MovieRenderPipelineCore/Private/Graph/MovieGraphConfig.cpp
@@ -19,6 +19,7 @@
 #include "Graph/MovieGraphUtils.h"
 #include "MovieGraphUtils.h"
 #include "MoviePipelineQueue.h"
+#include "MoviePipelineUtils.h"
 #include "MovieRenderPipelineCoreModule.h"
 #include "UObject/Package.h"
 
@@ -419,6 +420,49 @@ void UMovieGraphOutput::PostEditChangeProperty(FPropertyChangedEvent& PropertyCh
 }
 #endif // WITH_EDITOR
 
+bool UMovieGraphEvaluatedConfig::GetVariableValueByName(const FString& InVariableName, UMovieGraphValueContainer*& OutVariableValue) const
+{
+	if (const TArray<UMovieGraphValueContainer*>* Variables = VariablesByName.Find(InVariableName))
+	{
+		if (Variables->Num() > 1)
+		{
+			MRG_LOG_ONCE_PER_RENDER(DuplicateVariableNameWarning, LogMovieRenderPipeline, Warning,
+				TEXT("GetVariableValueByName(): Found a duplicate variable name [%s] within the graphs being used to render (including subgraphs). "
+					 "It's suggested that variable names are unique so their values can be fetched reliably from within blueprints."), *InVariableName);
+		}
+
+		UMovieGraphValueContainer* FirstVariableWithName = (*Variables)[0];
+		check(FirstVariableWithName);
+
+		OutVariableValue = FirstVariableWithName;
+		return true;
+	}
+
+	// It's not ideal to allocate a new "default" value container in the case of an invalid variable name. However, doing this dramatically simplifies
+	// blueprints (eg, you don't have to do complicated control flow branching based on the true/false return value of this function to avoid
+	// evaluating a null OutVariableValue) and the allocation is arguably worth the convenience.
+	OutVariableValue = NewObject<UMovieGraphValueContainer>(GetTransientPackage());
+	return false;
+}
+
+bool UMovieGraphEvaluatedConfig::GetVariableStringValueByName(const FString& InVariableName, const FString& InDefaultValue, FString& OutVariableValue) const
+{
+	UMovieGraphValueContainer* ValueContainer;
+	if (GetVariableValueByName(InVariableName, ValueContainer))
+	{
+		OutVariableValue = ValueContainer->GetValueSerializedString();
+		return true;
+	}
+
+	OutVariableValue = InDefaultValue;
+	return false;
+}
+
+const TArray<UMovieGraphConfig*>& UMovieGraphEvaluatedConfig::GetEvaluatedGraphs() const
+{
+	return EvaluatedGraphs;
+}
+
 UMovieGraphConfig::UMovieGraphConfig()
 {
 	InputNode = CreateDefaultSubobject<UMovieGraphInputNode>(TEXT("DefaultInputNode"));
@@ -1818,10 +1862,13 @@ bool UMovieGraphConfig::CreateFlattenedGraph_Recursive(UMovieGraphEvaluatedConfi
 
 	// Add this node to the set of visited nodes so it can be checked for cycles later. Get the graph from GetTypedOuter() rather than "this" because
 	// this method will be called recursively, potentially on pins within subgraphs.
-	const UMovieGraphConfig* OwningGraph = InPinToFollow->GetTypedOuter<UMovieGraphConfig>();
+	UMovieGraphConfig* OwningGraph = InPinToFollow->GetTypedOuter<UMovieGraphConfig>();
 	ensure(OwningGraph);
 	TSet<TObjectPtr<UMovieGraphNode>>& VisitedNodeSet = InEvaluationContext.VisitedNodesByOwningGraph.FindOrAdd(OwningGraph).VisitedNodes;
 	VisitedNodeSet.Add(Node);
+
+	// The evaluated config also keeps track of all the graphs that were visited.
+	InOwningConfig->EvaluatedGraphs.AddUnique(OwningGraph);
 	
 	const bool bShouldIncludeNode =
 		Node->IsA<UMovieGraphSettingNode>() &&
@@ -2166,6 +2213,23 @@ UMovieGraphEvaluatedConfig* UMovieGraphConfig::CreateFlattenedGraph(const FMovie
 			}
 		}
 	}
+
+	// Update the evaluated config with all of the evaluated variable values. These values are typically only used by scripting or blueprints.
+	for (const UMovieGraphConfig* EvaluatedGraph : NewContext->EvaluatedGraphs)
+	{
+		for (UMovieGraphVariable* GraphVariable : EvaluatedGraph->GetVariables())
+		{
+			TArray<UMovieGraphValueContainer*>& VariablesWithName = NewContext->VariablesByName.FindOrAdd(GraphVariable->GetMemberName());
+
+			TObjectPtr<UMovieGraphValueContainer> VariableValue;
+			if (UMovieGraphVariableNode::GetResolvedVariableValue(GraphVariable, &InContext, VariableValue))
+			{
+				// *All* variables with the given name are added. This is done so that we can warn that multiple variables were found with a specific
+				// name. May also have future uses if we need to get a resolved variable value from a specific graph.
+				VariablesWithName.Add(VariableValue);
+			}
+		}
+	}
 	
 	bool bHasRenderLayerNode = false;
 
diff --git a/Engine/Plugins/MovieScene/MovieRenderPipeline/Source/MovieRenderPipelineCore/Private/Graph/Nodes/MovieGraphVariableNode.cpp b/Engine/Plugins/MovieScene/MovieRenderPipeline/Source/MovieRenderPipelineCore/Private/Graph/Nodes/MovieGraphVariableNode.cpp
--- a/Engine/Plugins/MovieScene/MovieRenderPipeline/Source/MovieRenderPipelineCore/Private/Graph/Nodes/MovieGraphVariableNode.cpp
+++ b/Engine/Plugins/MovieScene/MovieRenderPipeline/Source/MovieRenderPipelineCore/Private/Graph/Nodes/MovieGraphVariableNode.cpp
@@ -37,7 +37,7 @@ FString UMovieGraphVariableNode::GetResolvedValueForOutputPin(const FName& InPin
 	if (GraphVariable && (GraphVariable->GetMemberName() == InPinName))
 	{
 		TObjectPtr<UMovieJobVariableAssignmentContainer> VariableAssignment;
-		if (ContextHasEnabledAssignmentForVariable(InContext, VariableAssignment))
+		if (ContextHasEnabledAssignmentForVariable(GraphVariable, InContext, VariableAssignment))
 		{
 			return VariableAssignment->GetValueSerializedString(GraphVariable);
 		}
@@ -54,7 +54,7 @@ bool UMovieGraphVariableNode::GetResolvedValueForOutputPin(const FName& InPinNam
 	if (GraphVariable && (GraphVariable->GetMemberName() == InPinName))
 	{
 		TObjectPtr<UMovieJobVariableAssignmentContainer> VariableAssignment;
-		if (ContextHasEnabledAssignmentForVariable(InContext, VariableAssignment))
+		if (ContextHasEnabledAssignmentForVariable(GraphVariable, InContext, VariableAssignment))
 		{
 			return VariableAssignment->GetValueContainer(GraphVariable, OutValueContainer);
 		}
@@ -67,6 +67,11 @@ bool UMovieGraphVariableNode::GetResolvedValueForOutputPin(const FName& InPinNam
 	return false;
 }
 
+UMovieGraphVariable* UMovieGraphVariableNode::GetVariable() const
+{
+	return GraphVariable;
+}
+
 void UMovieGraphVariableNode::SetVariable(UMovieGraphVariable* InVariable)
 {
 	if (InVariable)
@@ -81,6 +86,24 @@ void UMovieGraphVariableNode::SetVariable(UMovieGraphVariable* InVariable)
 	}
 }
 
+bool UMovieGraphVariableNode::IsGlobalVariable() const
+{
+	return GraphVariable && GraphVariable->IsGlobal();
+}
+
+bool UMovieGraphVariableNode::GetResolvedVariableValue(const TObjectPtr<UMovieGraphVariable> InGraphVariable, const FMovieGraphTraversalContext* InContext, TObjectPtr<UMovieGraphValueContainer>& OutValueContainer)
+{
+	TObjectPtr<UMovieJobVariableAssignmentContainer> VariableAssignment;
+	if (ContextHasEnabledAssignmentForVariable(InGraphVariable.Get(), InContext, VariableAssignment))
+	{
+		return VariableAssignment->GetValueContainer(InGraphVariable.Get(), OutValueContainer);
+	}
+	
+	// No valid variable assignment: just get the value from the variable
+	OutValueContainer = InGraphVariable;
+	return true;
+}
+
 #if WITH_EDITOR
 FText UMovieGraphVariableNode::GetNodeTitle(const bool bGetDescriptive) const
 {
@@ -133,7 +156,7 @@ void UMovieGraphVariableNode::UpdateOutputPin(UMovieGraphMember* ChangedVariable
 	OnNodeChangedDelegate.Broadcast(this);
 }
 
-bool UMovieGraphVariableNode::ContextHasEnabledAssignmentForVariable(const FMovieGraphTraversalContext* InContext, TObjectPtr<UMovieJobVariableAssignmentContainer>& OutVariableAssignment) const
+bool UMovieGraphVariableNode::ContextHasEnabledAssignmentForVariable(const TObjectPtr<UMovieGraphVariable> InGraphVariable, const FMovieGraphTraversalContext* InContext, TObjectPtr<UMovieJobVariableAssignmentContainer>& OutVariableAssignment)
 {
 	if (!InContext)
 	{
@@ -149,24 +172,24 @@ bool UMovieGraphVariableNode::ContextHasEnabledAssignmentForVariable(const FMovi
 	
 	if (Shot)
 	{
-		ShotVariableAssignments = Shot->GetOrCreateJobVariableAssignmentsForGraph(GetGraph());
+		ShotVariableAssignments = Shot->GetOrCreateJobVariableAssignmentsForGraph(InGraphVariable->GetOwningGraph());
 		
 		// The shot can also override variables on the primary job's graph (in addition to the shot-level ones fetched/created above)
 		if (PrimaryJob)
 		{
 			constexpr bool bIsForPrimaryOverrides = true;
-			ShotVariableAssignments_PrimaryOverrides = Shot->GetOrCreateJobVariableAssignmentsForGraph(GetGraph(), bIsForPrimaryOverrides);
+			ShotVariableAssignments_PrimaryOverrides = Shot->GetOrCreateJobVariableAssignmentsForGraph(InGraphVariable->GetOwningGraph(), bIsForPrimaryOverrides);
 		}
 	}
 
 	if (PrimaryJob)
 	{
-		JobVariableAssignments = PrimaryJob->GetOrCreateJobVariableAssignmentsForGraph(GetGraph());
+		JobVariableAssignments = PrimaryJob->GetOrCreateJobVariableAssignmentsForGraph(InGraphVariable->GetOwningGraph());
 	}
 	
 	// Check the shot job first for an enabled job variable assignment for this variable. Shot jobs take precedence over primary jobs.
 	bool bIsEnabled = false;
-	if (ShotVariableAssignments && ShotVariableAssignments->GetVariableAssignmentEnableState(GraphVariable, bIsEnabled))
+	if (ShotVariableAssignments && ShotVariableAssignments->GetVariableAssignmentEnableState(InGraphVariable.Get(), bIsEnabled))
 	{
 		if (bIsEnabled)
 		{
@@ -176,7 +199,7 @@ bool UMovieGraphVariableNode::ContextHasEnabledAssignmentForVariable(const FMovi
 	}
 
 	// Next check for shot-level overrides to the primary graph variables.
-	if (ShotVariableAssignments_PrimaryOverrides && ShotVariableAssignments_PrimaryOverrides->GetVariableAssignmentEnableState(GraphVariable, bIsEnabled))
+	if (ShotVariableAssignments_PrimaryOverrides && ShotVariableAssignments_PrimaryOverrides->GetVariableAssignmentEnableState(InGraphVariable.Get(), bIsEnabled))
 	{
 		if (bIsEnabled)
 		{
@@ -186,7 +209,7 @@ bool UMovieGraphVariableNode::ContextHasEnabledAssignmentForVariable(const FMovi
 	}
 
 	// Check the primary job last.
-	if (JobVariableAssignments && JobVariableAssignments->GetVariableAssignmentEnableState(GraphVariable, bIsEnabled))
+	if (JobVariableAssignments && JobVariableAssignments->GetVariableAssignmentEnableState(InGraphVariable.Get(), bIsEnabled))
 	{
 		if (bIsEnabled)
 		{
diff --git a/Engine/Plugins/MovieScene/MovieRenderPipeline/Source/MovieRenderPipelineCore/Public/Graph/MovieGraphConfig.h b/Engine/Plugins/MovieScene/MovieRenderPipeline/Source/MovieRenderPipelineCore/Public/Graph/MovieGraphConfig.h
--- a/Engine/Plugins/MovieScene/MovieRenderPipeline/Source/MovieRenderPipelineCore/Public/Graph/MovieGraphConfig.h
+++ b/Engine/Plugins/MovieScene/MovieRenderPipeline/Source/MovieRenderPipelineCore/Public/Graph/MovieGraphConfig.h
@@ -553,11 +553,56 @@ public:
 
 		return ResultNodes;
 	}
+
+	/**
+	 * Gets a variable's value via a "value container". This container lets you fetch the variable's value in its native form, rather than being
+	 * forced into a string (see GetVariableStringValueByName()). The variable's value is returned from the container via any of its GetValue*()
+	 * methods, like GetValueFloat().
+	 *
+	 * Returns true if a variable with the provided name exists, else false. If the variable does not exist, the value container will be initialized
+	 * to an empty one, which defaults to containing a float with value 0. The evaluated graph does not store a reference to this value container, so
+	 * if used outside of blueprints, be sure to keep a reference to it in order to prevent it from being garbage collected.
+	 */
+	UFUNCTION(BlueprintCallable, Category = "Movie Graph", meta=(ReturnDisplayName = "Variable Exists"))
+	bool GetVariableValueByName(
+		UPARAM(DisplayName = "Variable Name") const FString& InVariableName,
+		UPARAM(DisplayName = "Value Container") UMovieGraphValueContainer*& OutVariableValue) const;
+
+	/**
+	 * Convenience function to get a variable's string value, instead of its value container (see GetVariableValueByName()). Uses the provided
+	 * default value if a variable by the given name does not exist. For variables that aren't string-typed, this provides the serialized
+	 * representation of the variable value. Returns true if a variable with the provided name exists, else false.
+	 */
+	UFUNCTION(BlueprintCallable, Category = "Movie Graph", meta=(ReturnDisplayName = "Variable Exists"))
+	bool GetVariableStringValueByName(
+		UPARAM(DisplayName = "Variable Name") const FString& InVariableName,
+		UPARAM(DisplayName = "Default Value") const FString& InDefaultValue,
+		UPARAM(DisplayName = "String Value") FString& OutVariableValue) const;
+
+	/** Gets all of the graphs that were included in the evaluation process. */
+	UFUNCTION(BlueprintCallable, Category = "Movie Graph")
+	const TArray<UMovieGraphConfig*>& GetEvaluatedGraphs() const;
 	
 public:
 	/** Mapping between named branches (at the root of the config) and their evaluated values. */
 	UPROPERTY(Transient)
 	TMap<FName, FMovieGraphEvaluatedBranchConfig> BranchConfigMapping;
+
+private:
+	// Allow the graph traversal/evaluation process to update some private members.
+	friend class UMovieGraphConfig;
+
+	/**
+	 * The fully-evaluated values of all variables used in all graphs traversed in the evaluated result. The key is the variable name,
+	 * and the value is the resolved value of the variable. If there is more than one entry in the array, that means that multiple graphs
+	 * used the same variable name. This is not an error, but it makes fetching a variable by name challenging without any additional context.
+	 * For future use, all variables with the same name are gathered, where the first in the array is the variable from the parent-most graph,
+	 * and so on down the chain of subgraphs.
+	 */
+	TMap<FString, TArray<UMovieGraphValueContainer*>> VariablesByName;
+
+	/** All of the graphs that were included in the evaluation process. */
+	TArray<UMovieGraphConfig*> EvaluatedGraphs;
 };
 
 UINTERFACE(MinimalAPI)
diff --git a/Engine/Plugins/MovieScene/MovieRenderPipeline/Source/MovieRenderPipelineCore/Public/Graph/Nodes/MovieGraphVariableNode.h b/Engine/Plugins/MovieScene/MovieRenderPipeline/Source/MovieRenderPipelineCore/Public/Graph/Nodes/MovieGraphVariableNode.h
--- a/Engine/Plugins/MovieScene/MovieRenderPipeline/Source/MovieRenderPipelineCore/Public/Graph/Nodes/MovieGraphVariableNode.h
+++ b/Engine/Plugins/MovieScene/MovieRenderPipeline/Source/MovieRenderPipelineCore/Public/Graph/Nodes/MovieGraphVariableNode.h
@@ -30,14 +30,21 @@ public:
 
 	/** Gets the variable that this node represents. */
 	UFUNCTION(BlueprintCallable, Category = "Movie Graph")
-	UMovieGraphVariable* GetVariable() const { return GraphVariable; }
+	UE_API UMovieGraphVariable* GetVariable() const;
 
 	/** Sets the variable that this node represents. */
 	UFUNCTION(BlueprintCallable, Category = "Movie Graph")
 	UE_API void SetVariable(UMovieGraphVariable* InVariable);
 
 	/** Returns true if this node represents a global variable, else false. */
-	bool IsGlobalVariable() const { return GraphVariable && GraphVariable->IsGlobal(); }
+	UFUNCTION(BlueprintCallable, Category = "Movie Graph")
+	UE_API bool IsGlobalVariable() const;
+
+	/**
+	 * Resolves the provided variable, given a specific traversal context (needed in order to account for job-level overrides). If the variable value
+	 * could be resolved, returns true and provides the value in OutValueContainer, otherwise returns false.
+	 */
+	UE_API static bool GetResolvedVariableValue(const TObjectPtr<UMovieGraphVariable> InGraphVariable, const FMovieGraphTraversalContext* InContext, TObjectPtr<UMovieGraphValueContainer>& OutValueContainer);
 
 #if WITH_EDITOR
 	UE_API virtual FText GetNodeTitle(const bool bGetDescriptive = false) const override;
@@ -52,8 +59,8 @@ private:
 	/** Updates the output pin on the node to match the provided variable. */
 	UE_API void UpdateOutputPin(UMovieGraphMember* ChangedVariable) const;
 
-	/** Determines if the job in the given context has a valid, enabled variable assignment for the variable this node represents. Sets OutVariableAssignment if it does. */
-	UE_API bool ContextHasEnabledAssignmentForVariable(const FMovieGraphTraversalContext* InContext, TObjectPtr<UMovieJobVariableAssignmentContainer>& OutVariableAssignment) const;
+	/** Determines if the job in the given context has a valid, enabled variable assignment for the provided variable. Sets OutVariableAssignment if it does. */
+	UE_API static bool ContextHasEnabledAssignmentForVariable(const TObjectPtr<UMovieGraphVariable> InGraphVariable, const FMovieGraphTraversalContext* InContext, TObjectPtr<UMovieJobVariableAssignmentContainer>& OutVariableAssignment);
 
 private:
 	/** The underlying graph variable this node represents. */
diff --git a/Engine/Plugins/MovieScene/MovieRenderPipeline/Source/MovieRenderPipelineCore/Public/MoviePipelineUtils.h b/Engine/Plugins/MovieScene/MovieRenderPipeline/Source/MovieRenderPipelineCore/Public/MoviePipelineUtils.h
--- a/Engine/Plugins/MovieScene/MovieRenderPipeline/Source/MovieRenderPipelineCore/Public/MoviePipelineUtils.h
+++ b/Engine/Plugins/MovieScene/MovieRenderPipeline/Source/MovieRenderPipelineCore/Public/MoviePipelineUtils.h
@@ -159,6 +159,33 @@ MOVIERENDERPIPELINECORE_API void GetPassCompositeData(FMoviePipelineMergerOutput
 	} \
 }
 
+/*
+ * Log a message once per render.
+ *
+ * Note that the arguments to this are the same as for UE_LOG, but also takes an additional `Identifier`. The identifier is what uniquely identifies
+ * this specific message within the current SCOPE (identifiers are NOT unique across scopes, so the same log message will be logged more than once per
+ * render if it appears in different scopes).
+ */
+#define MRG_LOG_ONCE_PER_RENDER(Identifier, LogCategory, LogVerbosity, LogFormat, ...)							\
+	static UWorld* PREPROCESSOR_JOIN(__LastRenderWorld__, Identifier) = nullptr;								\
+	static TArray<FString> PREPROCESSOR_JOIN(__LoggedMessages__, Identifier);									\
+																												\
+	/* If this world has not been encountered yet (ie, no renders yet, or this is a new render), clear out */	\
+	/* the messages that have been recorded as already logged. */												\
+	if (PREPROCESSOR_JOIN(__LastRenderWorld__, Identifier) != MoviePipeline::FindCurrentWorld())				\
+	{																											\
+		PREPROCESSOR_JOIN(__LastRenderWorld__, Identifier) = MoviePipeline::FindCurrentWorld();					\
+		PREPROCESSOR_JOIN(__LoggedMessages__, Identifier).Empty();												\
+	}																											\
+																												\
+	/* If this message hasn't been logged, record that it has been logged so we can skip it for other frames. */\
+	const FString FinalLogMessage = FString::Printf(LogFormat, ##__VA_ARGS__);									\
+	if (!PREPROCESSOR_JOIN(__LoggedMessages__, Identifier).Contains(FinalLogMessage))							\
+	{																											\
+		UE_LOG(LogCategory, LogVerbosity, TEXT("%s"), *FinalLogMessage);										\
+		PREPROCESSOR_JOIN(__LoggedMessages__, Identifier).Add(FinalLogMessage);									\
+	}																											\
+
 namespace UE
 {
 	namespace MovieRenderPipeline
diff --git a/Engine/Plugins/MovieScene/MovieRenderPipeline/Source/MovieRenderPipelineRenderPasses/Private/MovieGraphImageSequenceOutputNode.cpp b/Engine/Plugins/MovieScene/MovieRenderPipeline/Source/MovieRenderPipelineRenderPasses/Private/MovieGraphImageSequenceOutputNode.cpp
--- a/Engine/Plugins/MovieScene/MovieRenderPipeline/Source/MovieRenderPipelineRenderPasses/Private/MovieGraphImageSequenceOutputNode.cpp
+++ b/Engine/Plugins/MovieScene/MovieRenderPipeline/Source/MovieRenderPipelineRenderPasses/Private/MovieGraphImageSequenceOutputNode.cpp
@@ -23,21 +23,6 @@
 
 #include UE_INLINE_GENERATED_CPP_BY_NAME(MovieGraphImageSequenceOutputNode)
 
-/*
- * Log a message once per render.
- *
- * Note that the arguments to this are the same as for UE_LOG, but also takes an additional `Identifier`. The identifier is what uniquely identifies
- * this specific message within the current SCOPE (identifiers are NOT unique across scopes, so the same log message will be logged more than once per
- * render if it appears in different scopes).
- */
-#define MRG_LOG_ONCE_PER_RENDER(Identifier, LogCategory, LogVerbosity, LogFormat, ...)							\
-	static UWorld* PREPROCESSOR_JOIN(LastRenderWorld_, Identifier) = nullptr;									\
-	if (PREPROCESSOR_JOIN(LastRenderWorld_, Identifier) != MoviePipeline::FindCurrentWorld())					\
-	{																											\
-		PREPROCESSOR_JOIN(LastRenderWorld_, Identifier) = MoviePipeline::FindCurrentWorld();					\
-		UE_LOG(LogCategory, LogVerbosity, LogFormat, ##__VA_ARGS__);											\
-	}																											\
-
 UMovieGraphImageSequenceOutputNode::UMovieGraphImageSequenceOutputNode()
 {
 	ImageWriteQueue = &FModuleManager::Get().LoadModuleChecked<IImageWriteQueueModule>("ImageWriteQueue").GetWriteQueue();
