MoviePipeline: Fixed an issue where Quick Render's "Current Viewport" and "Selected Cameras" modes would not pick up the effects of the active level sequence. These modes will also now pick up the current playhead position in the active sequence.

On the scripting side, UMovieGraphQuickRenderModeSettings::LevelSequenceOverride is also now respected in these modes.

diff --git a/Engine/Plugins/MovieScene/MovieRenderPipeline/Source/MovieRenderPipelineCore/Public/Graph/MovieGraphQuickRenderSettings.h b/Engine/Plugins/MovieScene/MovieRenderPipeline/Source/MovieRenderPipelineCore/Public/Graph/MovieGraphQuickRenderSettings.h
--- a/Engine/Plugins/MovieScene/MovieRenderPipeline/Source/MovieRenderPipelineCore/Public/Graph/MovieGraphQuickRenderSettings.h
+++ b/Engine/Plugins/MovieScene/MovieRenderPipeline/Source/MovieRenderPipelineCore/Public/Graph/MovieGraphQuickRenderSettings.h
@@ -157,9 +157,9 @@ public:
 	TSoftObjectPtr<UMovieGraphConfig> GraphPreset;
 	
 	/**
-	 * For sequence-centric modes (like "Current Sequence"), if this override is set, this sequence will be used instead of the level sequence that's
-	 * active in Sequencer. In UI-based Quick Render, this property is never set directly because Sequencer is always used. However, for scripting
-	 * purposes, there may be cases where Sequencer is not open, and a specific level sequence needs to be used.
+	 * If this override is set, this sequence will be used instead of the level sequence that's active in Sequencer. In UI-based Quick Render, this
+	 * property is never set directly because Sequencer is always used. However, for scripting purposes, there may be cases where Sequencer is not
+	 * open, and a specific level sequence needs to be used.
 	 */
 	UPROPERTY(BlueprintReadWrite, Category="Configuration", DisplayName="Movie Render Graph Config")
 	TSoftObjectPtr<class ULevelSequence> LevelSequenceOverride = nullptr;
diff --git a/Engine/Plugins/MovieScene/MovieRenderPipeline/Source/MovieRenderPipelineEditor/Private/Graph/MovieGraphQuickRender.cpp b/Engine/Plugins/MovieScene/MovieRenderPipeline/Source/MovieRenderPipelineEditor/Private/Graph/MovieGraphQuickRender.cpp
--- a/Engine/Plugins/MovieScene/MovieRenderPipeline/Source/MovieRenderPipelineEditor/Private/Graph/MovieGraphQuickRender.cpp
+++ b/Engine/Plugins/MovieScene/MovieRenderPipeline/Source/MovieRenderPipelineEditor/Private/Graph/MovieGraphQuickRender.cpp
@@ -21,6 +21,7 @@
 #include "MovieRenderPipelineCoreModule.h"
 #include "MovieRenderPipelineSettings.h"
 #include "MovieScene.h"
+#include "MovieSceneSequencePlayer.h"
 #include "MovieSceneSpawnableAnnotation.h"
 #include "ObjectEditorUtils.h"
 #include "Scalability.h"
@@ -100,6 +101,33 @@ namespace
 		SavePackageHelper(NewPackage, *PackageFilename, RF_Standalone, GWarn, SAVE_None);
 		GEditor->GetEditorSubsystem<UAssetEditorSubsystem>()->OpenEditorForAsset(NewLevelSequence);
 	}
+
+	/**
+	 * Gets the override level sequence from the mode settings (set via scripting), or the active level sequence if no override was set. If no
+	 * override was set, and there's no active level sequence, a warning dialog box will be shown if `bWarn` is true.
+	 */
+	ULevelSequence* GetCurrentLevelSequenceOrWarn(const TObjectPtr<const UMovieGraphQuickRenderModeSettings>& InModeSettings, EMovieGraphQuickRenderMode ActiveQuickRenderMode, const bool bWarn = true)
+	{
+		// Scripting may override the level sequence to use
+		ULevelSequence* LevelSequenceOverride = InModeSettings->LevelSequenceOverride.LoadSynchronous();
+		if (IsValid(LevelSequenceOverride))
+		{
+			return LevelSequenceOverride;
+		}
+		
+		ULevelSequence* CurrentLevelSequence = ULevelSequenceEditorBlueprintLibrary::GetCurrentLevelSequence();
+		if (!CurrentLevelSequence && bWarn)
+		{
+			const FText ModeName = StaticEnum<EMovieGraphQuickRenderMode>()->GetDisplayNameTextByValue(static_cast<int64>(ActiveQuickRenderMode));
+			const FText MessageText = LOCTEXT("NoActiveSequenceWarning", "Quick Render needs a level sequence opened in Sequencer in order to function in the current mode [{0}]. Open a level sequence, then try again.");
+		
+			FMessageDialog::Open(EAppMsgType::Ok, FText::Format(MessageText, ModeName));
+
+			return nullptr;
+		}
+
+		return CurrentLevelSequence;
+	};
 }
 
 /**
@@ -143,6 +171,17 @@ void UMovieGraphQuickRenderSubsystem::BeginQuickRender(const EMovieGraphQuickRen
 	// Store a reference to the settings used, so the executor callbacks can rely on them still existing.
 	QuickRenderModeSettings = InQuickRenderSettings;
 
+	// Find the editor world
+	UWorld* EditorWorld = GetWorldOfType(EWorldType::Editor);
+	if (!EditorWorld)
+	{
+		UE_LOG(LogMovieRenderPipeline, Error, TEXT("Unable to Quick Render: Could not find a currently active editor world."));
+		return;
+	}
+	
+	// Cache out anything necessary before PIE starts
+	CachePrePieData(EditorWorld);
+
 	// Set up the rendering and "utility" level sequences. These are both required to do a render successfully.
 	if (!SetUpAllLevelSequences())
 	{
@@ -155,14 +194,6 @@ void UMovieGraphQuickRenderSubsystem::BeginQuickRender(const EMovieGraphQuickRen
 		return;
 	}
 
-	// Find the editor world
-	UWorld* EditorWorld = GetWorldOfType(EWorldType::Editor);
-	if (!EditorWorld)
-	{
-		UE_LOG(LogMovieRenderPipeline, Error, TEXT("Unable to Quick Render: Could not find a currently active editor world."));
-		return;
-	}
-
 	// Generate a new queue, fully populated with a job to render from, the graph assigned to that job, and job-level variable overrides applied.
 	TObjectPtr<UMoviePipelineExecutorJob> NewJob = nullptr;
 	if (!GenerateAndPopulateQueue(TemporaryQueue, NewJob, QuickRenderModeSettings))
@@ -187,9 +218,6 @@ void UMovieGraphQuickRenderSubsystem::BeginQuickRender(const EMovieGraphQuickRen
 		return;
 	}
 
-	// Cache out anything necessary before PIE starts
-	CachePrePieData(EditorWorld);
-
 	// Update the graph as needed in order to properly apply the settings specified by the user. This will perform a graph evaluation, so the job
 	// should be fully initialized at this point.
 	ApplyQuickRenderUpdatesToDuplicateGraph(EditorWorld);
@@ -873,41 +901,39 @@ void UMovieGraphQuickRenderSubsystem::RestorePreRenderState()
 
 ULevelSequence* UMovieGraphQuickRenderSubsystem::SetUpRenderingLevelSequence() const
 {
-	auto GetCurrentLevelSequenceOrWarn = [](EMovieGraphQuickRenderMode ActiveQuickRenderMode) -> ULevelSequence*
+	auto CreateTemporaryLevelSequence = []() -> ULevelSequence*
 	{
-		ULevelSequence* CurrentLevelSequence = ULevelSequenceEditorBlueprintLibrary::GetCurrentLevelSequence();
-		if (!CurrentLevelSequence)
-		{
-			const FText ModeName = StaticEnum<EMovieGraphQuickRenderMode>()->GetDisplayNameTextByValue(static_cast<int64>(ActiveQuickRenderMode));
-			const FText MessageText = LOCTEXT("NoActiveSequenceWarning", "Quick Render needs a level sequence opened in Sequencer in order to function in the current mode [{0}]. Open a level sequence, then try again.");
-		
-			FMessageDialog::Open(EAppMsgType::Ok, FText::Format(MessageText, ModeName));
+		ULevelSequence* TemporaryLevelSequence = NewObject<ULevelSequence>(GetTransientPackage(), TEXT("QuickRenderSequence"), RF_Transient);
+		TemporaryLevelSequence->Initialize();
 
-			return nullptr;
-		}
-
-		return CurrentLevelSequence;
+		return TemporaryLevelSequence;
 	};
 	
-	// For the CurrentViewport and SelectedCameras modes, a temporary level sequence will be used to drive the render.
-	if ((QuickRenderMode == EMovieGraphQuickRenderMode::CurrentViewport) || (QuickRenderMode == EMovieGraphQuickRenderMode::SelectedCameras))
+	// For the Current Viewport mode, if a level sequence is active, use that to drive the render (even though this mode generates stills, the level
+	// sequence can impact transforms, etc which need to be captured). If no level sequence is active, use a temporary level sequence.
+	if (QuickRenderMode == EMovieGraphQuickRenderMode::CurrentViewport)
 	{
-		// Create the new, temporary level sequence
-		ULevelSequence* QuickRenderSequence = NewObject<ULevelSequence>(GetTransientPackage(), TEXT("QuickRenderSequence"), RF_Transient);
-		QuickRenderSequence->Initialize();
+		constexpr bool bWarn = false;
+		ULevelSequence* CurrentLevelSequence = GetCurrentLevelSequenceOrWarn(QuickRenderModeSettings, QuickRenderMode, bWarn);
+		if (IsValid(CurrentLevelSequence))
+		{
+			return CurrentLevelSequence;
+		}
 
-		return QuickRenderSequence;
+		return CreateTemporaryLevelSequence();
 	}
 
-	// Scripting may override the level sequence to use
-	ULevelSequence* LevelSequenceOverride = QuickRenderModeSettings->LevelSequenceOverride.LoadSynchronous();
-	if (IsValid(LevelSequenceOverride))
+	// For the Selected Cameras mode, a temporary level sequence will always be rendered from. This temporary sequence will ultimately have the
+	// utility level sequence added to it as a subsequence, and that will control the camera cuts as well as include the active level sequence (if
+	// there is one). This mode needs more control over the active level sequence if it exists, which is why it is not used as the rendering level
+	// sequence directly (see comments in SetUpUtilityLevelSequence() for more details).
+	if (QuickRenderMode == EMovieGraphQuickRenderMode::SelectedCameras)
 	{
-		return LevelSequenceOverride;
+		return CreateTemporaryLevelSequence();
 	}
 
-	// Otherwise, the active mode should use the level sequence that's currently active in Sequencer.
-	ULevelSequence* CurrentLevelSequence = GetCurrentLevelSequenceOrWarn(QuickRenderMode);
+	// Otherwise, use the level sequence that's specified via scripting, or currently active in Sequencer.
+	ULevelSequence* CurrentLevelSequence = GetCurrentLevelSequenceOrWarn(QuickRenderModeSettings, QuickRenderMode);
 	if (!CurrentLevelSequence)
 	{
 		return nullptr;
@@ -921,44 +947,42 @@ ULevelSequence* UMovieGraphQuickRenderSubsystem::SetUpUtilityLevelSequence() con
 	ULevelSequence* UtilitySequence = NewObject<ULevelSequence>(GetTransientPackage(), TEXT("QuickRenderUtilitySequence"), RF_Transient);
 	UtilitySequence->Initialize();
 
-	// If the currently active level sequence in Sequencer has spawnables, mirror those spawnable bindings into the utility level sequence. For modes
-	// that don't use the active level sequence, we need to do this in order to match visibility.
-	auto AddSpawnableTracks = [this, UtilitySequence]() -> void
+	// For the Selected Cameras mode, we need to add the level sequence that's active (if there is one) as a subsequence. This needs to be done once
+	// PER camera. This is strange, but it's done in order to accomodate the current playhead of the active level sequence. If the user is on frame
+	// 100 for example, EACH camera needs to render frame 100. The camera cuts generated for this mode work such that there's one camera cut per frame
+	// in the utility sequence, so without adding the level sequence as a subsequence once PER camera, you'd wind up with camera 1 rendering frame
+	// 100, camera 2 rendering frame 101, and so on. The level sequence section that's added per camera is set up to start at the current playhead, so
+	// the effect is that each camera renders the active sequence at the current playhead. It's also given a very negative bias so the camera cuts
+	// added to the utility sequence have priority over any that may exist in the active level sequence.
+	//
+	// An alternative to this is using the "Render all cameras" mode, but that has the potential to use significantly more memory. If rendering many
+	// cameras at once, the approach taken here the potential to be much more efficient.
+	if (QuickRenderMode == EMovieGraphQuickRenderMode::SelectedCameras)
 	{
-		const UWorld* EditorWorld = GetWorldOfType(EWorldType::Editor);
-		const UMovieScene* UtilityMovieScene = UtilitySequence->GetMovieScene();
-
-		for (FActorIterator It(EditorWorld); It; ++It)
+		constexpr bool bWarn = false;
+		ULevelSequence* CurrentLevelSequence = GetCurrentLevelSequenceOrWarn(QuickRenderModeSettings, QuickRenderMode, bWarn);
+		
+		if (CurrentLevelSequence)
 		{
-			AActor* Actor = *It;
+			UMovieScene* MovieScene = UtilitySequence->GetMovieScene();
+			const FFrameRate DisplayRate = MovieScene->GetDisplayRate();
+			const FFrameRate DestinationRate = MovieScene->GetTickResolution();
 
-			// There's no need to make the actor spawnable if it's currently not visible in the editor
-			if (!Actor || Actor->IsHiddenEd())
-			{
-				continue;
-			}
+			// The playhead position in the root level sequence
+			const FMovieSceneSequencePlaybackParams GlobalPosition = ULevelSequenceEditorBlueprintLibrary::GetGlobalPosition();
 
-			TOptional<FMovieSceneSpawnableAnnotation> Spawnable = FMovieSceneSpawnableAnnotation::Find(Actor);
-			if (Spawnable.IsSet())
+			// Add the a section for the active level sequence for each camera
+			for (int32 CameraIndex = 0; CameraIndex < CachedPrePieData.SelectedCameras.Num(); ++CameraIndex)
 			{
-				FMovieGraphTransactionBlocker BlockTransactions;
-				
-				UE::Sequencer::FCreateBindingParams CreateBindingParams;
-				CreateBindingParams.bAllowCustomBinding = true;
-				CreateBindingParams.bSpawnable = true;
-				FSequencerUtilities::CreateOrReplaceBinding(nullptr, UtilityMovieScene->GetTypedOuter<UMovieSceneSequence>(), Actor, CreateBindingParams);
+				UMovieSceneSubTrack* NewSubsequenceTrack = CastChecked<UMovieSceneSubTrack>(MovieScene->AddTrack(UMovieSceneSubTrack::StaticClass()));
+				UMovieSceneSubSection* SequenceSubsection = NewSubsequenceTrack->AddSequence(CurrentLevelSequence, 0, 0);
+				SequenceSubsection->Parameters.StartFrameOffset = FFrameRate::TransformTime(GlobalPosition.Frame.FrameNumber, DisplayRate, DestinationRate).GetFrame();
+				SequenceSubsection->Parameters.HierarchicalBias = -100000;
+				SequenceSubsection->SetRange(TRange<FFrameNumber>(
+					FFrameRate::TransformTime(FFrameTime(CameraIndex), DisplayRate, DestinationRate).GetFrame(),
+					FFrameRate::TransformTime(FFrameTime(CameraIndex + 1), DisplayRate, DestinationRate).GetFrame()));
 			}
 		}
-	};
-
-	if (IsViewportLookFlagActive(EMovieGraphQuickRenderViewportLookFlags::Visibility))
-	{
-		// Only modes that don't use the active level sequence need the spawnable bindings mirrored. Ideally this could be done in
-		// PerformPreRenderSetup(), but at that point it's too late to get spawnable state because PIE has started.
-		if ((QuickRenderMode == EMovieGraphQuickRenderMode::CurrentViewport) || (QuickRenderMode == EMovieGraphQuickRenderMode::SelectedCameras))
-		{
-			AddSpawnableTracks();
-		}
 	}
 
 	return UtilitySequence;
@@ -998,6 +1022,14 @@ TRange<FFrameNumber> UMovieGraphQuickRenderSubsystem::GetPlaybackRange()
 	// If we're in Current Viewport, the playback range is always one frame long (only a single frame is rendered).
 	if (QuickRenderMode == EMovieGraphQuickRenderMode::CurrentViewport)
 	{
+		// If rendering from the active level sequence, the playback position should be where the current playhead is.
+		if (RenderingLevelSequence == ULevelSequenceEditorBlueprintLibrary::GetCurrentLevelSequence())
+		{
+			const FMovieSceneSequencePlaybackParams GlobalPosition = ULevelSequenceEditorBlueprintLibrary::GetGlobalPosition();
+			
+			return TRange<FFrameNumber>::Inclusive(GlobalPosition.Frame.FrameNumber, GlobalPosition.Frame.FrameNumber.Value + 1);
+		}
+		
 		return TRange<FFrameNumber>::Inclusive(0, 1);
 	}
 
