MoviePipeline: Added a 32 bit -> 16 bit conversion path for images to fix a crash when trying to write ObjectIds to Apple ProRes files. This conversion is not mathmatically correct (as changing ObjectIds away from their source 32-bits makes them invalid), but this works around the crash. Additionally, fixed an incorrect result from a merge in the 16->16 bit case which could result in uninitialized memory in the output if srgb conversion was false. The 32-16 path doesn't support srgb conversion (as both FLinearColor/FFloat16Color are supposed to be in linear space) and this case is mostly to avoid crashes when using a 32 bit source into 16-bit outputs in MRQ.

diff --git a/Engine/Plugins/MovieScene/MovieRenderPipeline/Source/MovieRenderPipelineCore/Private/MoviePipelineImageQuantization.cpp b/Engine/Plugins/MovieScene/MovieRenderPipeline/Source/MovieRenderPipelineCore/Private/MoviePipelineImageQuantization.cpp
--- a/Engine/Plugins/MovieScene/MovieRenderPipeline/Source/MovieRenderPipelineCore/Private/MoviePipelineImageQuantization.cpp
+++ b/Engine/Plugins/MovieScene/MovieRenderPipeline/Source/MovieRenderPipelineCore/Private/MoviePipelineImageQuantization.cpp
@@ -391,21 +391,52 @@ static TUniquePtr<FImagePixelData> QuantizePixelDataTo16bpp(const FImagePixelDat
 		const float* LookupTablePtr = LookupTable.GetData();
 		
 		const FColor* DataPtr = (FColor*)(SrcRawDataPtr);
-		for (int64 Index = 0; Index < (RawSize.X * (int64)RawSize.Y); Index++)
 		if (bConvertToSrgb)
 		{
-			OutColors[Index].R = LookupTablePtr[DataPtr[Index].R];
-			OutColors[Index].G = LookupTablePtr[DataPtr[Index].G];
-			OutColors[Index].B = LookupTablePtr[DataPtr[Index].B];
+			for (int64 Index = 0; Index < (RawSize.X * (int64)RawSize.Y); Index++)
+			{
+				OutColors[Index].R = LookupTablePtr[DataPtr[Index].R];
+				OutColors[Index].G = LookupTablePtr[DataPtr[Index].G];
+				OutColors[Index].B = LookupTablePtr[DataPtr[Index].B];
 
-			// Alpha is linear and doesn't get converted.
-			OutColors[Index].A = DataPtr[Index].A / 255.f;
+				// Alpha is linear and doesn't get converted.
+				OutColors[Index].A = DataPtr[Index].A / 255.f;
+			}
+		}
+		else
+		{
+			for (int64 Index = 0; Index < (RawSize.X * (int64)RawSize.Y); Index++)
+			{
+				OutColors[Index].R = DataPtr[Index].R / 255.f;
+				OutColors[Index].G = DataPtr[Index].G / 255.f;
+				OutColors[Index].B = DataPtr[Index].B / 255.f;
+				OutColors[Index].A = DataPtr[Index].A / 255.f;
+			}
 		}
 
 		QuantizedPixelData = MakeUnique<TImagePixelData<FFloat16Color>>(RawSize, TArray64<FFloat16Color>(MoveTemp(OutColors)), InPayload);
 		break;
 	}
 	case 32:
+	{
+		int64 SizeInBytes = 0;
+		const void* SrcRawDataPtr = nullptr;
+		InPixelData->GetRawData(SrcRawDataPtr, SizeInBytes);
+
+		TArray64<FFloat16Color> OutColors;
+		OutColors.SetNumUninitialized(RawSize.X * (int64)RawSize.Y);
+
+		// This doesn't support the bConvertToSrgb flag because both FLinearColor and FFloat16Color should
+		// be linear.
+		const FLinearColor* DataPtr = (FLinearColor*)(SrcRawDataPtr);
+
+		for (int64 Index = 0; Index < (RawSize.X * (int64)RawSize.Y); Index++)
+		{
+			OutColors[Index] = FFloat16Color(DataPtr[Index]);
+		}
+		QuantizedPixelData = MakeUnique<TImagePixelData<FFloat16Color>>(RawSize, TArray64<FFloat16Color>(MoveTemp(OutColors)), InPayload);
+		break;
+	}
 	default:
 		// Unsupported source bit-depth, consider adding it!
 		check(false);
