MovieGraph: Fixed an issue when using time dilation and frame_number_rel, where frames created using shots after the first one would double-count the frame numbers (using both the offset of the shot in the original source timeline, plus the number of relative frames already rendered) resulting in inconsistent frame numbers with gaps between them.

When using frame_number_rel we should always use the time from the very first rendered shot, because that's where the internal counter starts it's count up, so this keeps them in sync.

#virtualized

diff --git a/Engine/Plugins/MovieScene/MovieRenderPipeline/Source/MovieRenderPipelineCore/Private/Graph/MovieGraphFilenameResolveParams.cpp b/Engine/Plugins/MovieScene/MovieRenderPipeline/Source/MovieRenderPipelineCore/Private/Graph/MovieGraphFilenameResolveParams.cpp
--- a/Engine/Plugins/MovieScene/MovieRenderPipeline/Source/MovieRenderPipelineCore/Private/Graph/MovieGraphFilenameResolveParams.cpp
+++ b/Engine/Plugins/MovieScene/MovieRenderPipeline/Source/MovieRenderPipelineCore/Private/Graph/MovieGraphFilenameResolveParams.cpp
@@ -33,7 +33,12 @@ FMovieGraphFilenameResolveParams FMovieGraphFilenameResolveParams::MakeResolvePa
 			
 			// Calculate the offset (in frames) that relative frame numbers need to add to themselves so that they are correctly
 			// offset by the starting frame of the root sequence/shot (to match new updated relative behavior).
-			RootFrameNumberRelOffset = FFrameRate::TransformTime(Shot->ShotInfo.InitialTimeInRoot, Shot->ShotInfo.CachedTickResolution, Shot->ShotInfo.CachedFrameRate).FloorToFrame().Value;
+			// We use the first shot to calculate the root version of the relative frame number offset, because that'll have the correct
+			// root time for rendered frames (that should account for handle frames (included) vs warmup (not)). It's safe to assume index
+			// 0 exists here because we don't null check ShotIndex and it's either going to be >= 0.
+			const TObjectPtr<UMoviePipelineExecutorShot>& RootShot = InPipeline->GetActiveShotList()[0];
+			
+			RootFrameNumberRelOffset = FFrameRate::TransformTime(RootShot->ShotInfo.InitialTimeInRoot, RootShot->ShotInfo.CachedTickResolution, RootShot->ShotInfo.CachedFrameRate).FloorToFrame().Value;
 			ShotFrameNumberRelOffset = FFrameRate::TransformTime(Shot->ShotInfo.InitialTimeInShot, Shot->ShotInfo.CachedTickResolution, Shot->ShotInfo.CachedFrameRate).FloorToFrame().Value;
 		}
 	}
