MoviePipeline: Updated the UX for burn-ins in MRG (each output node can individually opt in or out of burn-ins).

With this change, the old burn-in node is now deprecated. That node cannot be used at the same time as the new output-specific burn-ins (a graph evaluation error will result).

diff --git a/Engine/Plugins/Media/AppleProResMedia/Source/AppleProResMedia/Private/MovieGraphAppleProResNode.cpp b/Engine/Plugins/Media/AppleProResMedia/Source/AppleProResMedia/Private/MovieGraphAppleProResNode.cpp
--- a/Engine/Plugins/Media/AppleProResMedia/Source/AppleProResMedia/Private/MovieGraphAppleProResNode.cpp
+++ b/Engine/Plugins/Media/AppleProResMedia/Source/AppleProResMedia/Private/MovieGraphAppleProResNode.cpp
@@ -7,6 +7,7 @@
 #include "Graph/MovieGraphConfig.h"
 #include "Graph/MovieGraphPipeline.h"
 #include "Graph/MovieGraphOCIOHelper.h"
+#include "Graph/Nodes/MovieGraphBurnInNode.h"
 #include "Graph/Nodes/MovieGraphGlobalOutputSettingNode.h"
 #include "ImageWriteTask.h"
 #include "MoviePipelineImageQuantization.h"
@@ -269,3 +270,25 @@ void UMovieGraphAppleProResNode::UpdateTelemetry(FMoviePipelineShotRenderTelemet
 {
 	InTelemetry->bUsesProRes = true;
 }
+
+void UMovieGraphAppleProResNode::InjectNodesPostEvaluation(const FName& InBranchName, UMovieGraphEvaluatedConfig* InEvaluatedConfig, TArray<UMovieGraphSettingNode*>& OutInjectedNodes)
+{
+	if (bEnableBurnIn)
+	{
+		UMovieGraphOutputBurnInNode* BurnInNode = NewObject<UMovieGraphOutputBurnInNode>(InEvaluatedConfig);
+		BurnInNode->bOverride_OutputName = true;
+		BurnInNode->bOverride_FileNameFormat = true;
+		BurnInNode->bOverride_OutputRestriction = true;
+		BurnInNode->bOverride_LayerNameFormat = true;
+		BurnInNode->bOverride_BurnInClass = true;
+		BurnInNode->bOverride_bCompositeOntoFinalImage = true;
+		BurnInNode->OutputName = TEXT("ProRes");
+		BurnInNode->FileNameFormat = BurnInFileNameFormat;
+		BurnInNode->OutputRestriction = FSoftClassPath(GetClass());
+		BurnInNode->LayerNameFormat = TEXT("");	// Unused
+		BurnInNode->BurnInClass = BurnInClass;
+		BurnInNode->bCompositeOntoFinalImage = bCompositeOntoFinalImage;
+
+		OutInjectedNodes.Add(BurnInNode);
+	}
+}
diff --git a/Engine/Plugins/Media/AppleProResMedia/Source/AppleProResMedia/Private/MovieGraphAppleProResNode.h b/Engine/Plugins/Media/AppleProResMedia/Source/AppleProResMedia/Private/MovieGraphAppleProResNode.h
--- a/Engine/Plugins/Media/AppleProResMedia/Source/AppleProResMedia/Private/MovieGraphAppleProResNode.h
+++ b/Engine/Plugins/Media/AppleProResMedia/Source/AppleProResMedia/Private/MovieGraphAppleProResNode.h
@@ -10,7 +10,7 @@
 
 /** A node which can output Apple ProRes movies. */
 UCLASS(BlueprintType)
-class UMovieGraphAppleProResNode : public UMovieGraphVideoOutputNode
+class UMovieGraphAppleProResNode : public UMovieGraphVideoOutputNode, public IMovieGraphEvaluationNodeInjector
 {
 	GENERATED_BODY()
 
@@ -40,6 +40,10 @@ protected:
 	virtual void UpdateTelemetry(FMoviePipelineShotRenderTelemetry* InTelemetry) const override;
 	// ~UMovieGraphSettingNode Interface
 
+	// IMovieGraphEvaluationNodeInjector Interface
+	virtual void InjectNodesPostEvaluation(const FName& InBranchName, UMovieGraphEvaluatedConfig* InEvaluatedConfig, TArray<UMovieGraphSettingNode*>& OutInjectedNodes) override;
+	// ~IMovieGraphEvaluationNodeInjector Interface
+
 protected:
 	struct FProResWriter : public MovieRenderGraph::IVideoCodecWriter
 	{
@@ -62,6 +66,18 @@ public:
 
 	UPROPERTY(EditAnywhere, BlueprintReadWrite, Category = Overrides, meta = (InlineEditConditionToggle))
 	uint8 bOverride_OCIOContext : 1;
+
+	UPROPERTY(EditAnywhere, BlueprintReadWrite, Category = Overrides, meta = (InlineEditConditionToggle))
+	uint8 bOverride_bEnableBurnIn : 1;
+
+	UPROPERTY(EditAnywhere, BlueprintReadWrite, Category = Overrides, meta = (InlineEditConditionToggle))
+	uint8 bOverride_BurnInClass : 1;
+
+	UPROPERTY(EditAnywhere, BlueprintReadWrite, Category = Overrides, meta = (InlineEditConditionToggle))
+	uint8 bOverride_bCompositeOntoFinalImage : 1;
+
+	UPROPERTY(EditAnywhere, BlueprintReadWrite, Category = Overrides, meta = (InlineEditConditionToggle))
+	uint8 bOverride_BurnInFileNameFormat : 1;
 	
 	/** The Apple ProRes codec that should be used. See Apple documentation for more specifics. Uses Rec 709 color primaries. */
 	UPROPERTY(EditAnywhere, BlueprintReadWrite, Category = "Apple ProRes", meta = (EditCondition = "bOverride_Quality"))
@@ -91,4 +107,20 @@ public:
 	*/
 	UPROPERTY(EditAnywhere, BlueprintReadWrite, Category = "OCIO", DisplayName = "OCIO Context", meta = (DisplayAfter = "OCIOConfiguration", EditCondition = "bOverride_OCIOContext"))
 	TMap<FString, FString> OCIOContext;
+
+	/** If true, this output node will also generate a burn-in. */
+	UPROPERTY(EditAnywhere, BlueprintReadWrite, Category = "Burn In", meta = (EditCondition = "bOverride_bEnableBurnIn"))
+	bool bEnableBurnIn = false;
+
+	/** The widget that the burn-in should use. */
+	UPROPERTY(EditAnywhere, BlueprintReadWrite, Category="Burn In", meta=(MetaClass="/Script/MovieRenderPipelineCore.MovieGraphBurnInWidget", EditCondition="bOverride_BurnInClass"))
+	FSoftClassPath BurnInClass = FSoftClassPath(TEXT("/MovieRenderPipeline/Blueprints/Graph/DefaultGraphBurnIn.DefaultGraphBurnIn_C"));
+
+	/** If true, the burn-in that's generated will be composited onto this output. Otherwise, the burn-in will be written to a different file. */
+	UPROPERTY(EditAnywhere, BlueprintReadWrite, Category = "Burn In", meta = (EditCondition = "bOverride_bCompositeOntoFinalImage"))
+	bool bCompositeOntoFinalImage = true;
+
+	/** The file name format used for writing out the burn-in video (if not composited). */
+	UPROPERTY(EditAnywhere, BlueprintReadWrite, Category = "Burn In", meta = (EditCondition = "bOverride_BurnInFileNameFormat"))
+	FString BurnInFileNameFormat = TEXT("{sequence_name}.{layer_name}.{renderer_name}");
 };
\ No newline at end of file
diff --git a/Engine/Plugins/Media/AvidDNxMedia/Source/Source/Private/MovieGraphAvidDNxHRNode.cpp b/Engine/Plugins/Media/AvidDNxMedia/Source/Source/Private/MovieGraphAvidDNxHRNode.cpp
--- a/Engine/Plugins/Media/AvidDNxMedia/Source/Source/Private/MovieGraphAvidDNxHRNode.cpp
+++ b/Engine/Plugins/Media/AvidDNxMedia/Source/Source/Private/MovieGraphAvidDNxHRNode.cpp
@@ -10,6 +10,7 @@
 #include "Graph/MovieGraphConfig.h"
 #include "Graph/MovieGraphPipeline.h"
 #include "Graph/MovieGraphOCIOHelper.h"
+#include "Graph/Nodes/MovieGraphBurnInNode.h"
 #include "Graph/Nodes/MovieGraphGlobalOutputSettingNode.h"
 #include "MovieRenderPipelineCoreModule.h"
 #include "Styling/AppStyle.h"
@@ -221,3 +222,25 @@ void UMovieGraphAvidDNxHRNode::UpdateTelemetry(FMoviePipelineShotRenderTelemetry
 {
 	InTelemetry->bUsesAvid = true;
 }
+
+void UMovieGraphAvidDNxHRNode::InjectNodesPostEvaluation(const FName& InBranchName, UMovieGraphEvaluatedConfig* InEvaluatedConfig, TArray<UMovieGraphSettingNode*>& OutInjectedNodes)
+{
+	if (bEnableBurnIn)
+	{
+		UMovieGraphOutputBurnInNode* BurnInNode = NewObject<UMovieGraphOutputBurnInNode>(InEvaluatedConfig);
+		BurnInNode->bOverride_OutputName = true;
+		BurnInNode->bOverride_FileNameFormat = true;
+		BurnInNode->bOverride_OutputRestriction = true;
+		BurnInNode->bOverride_LayerNameFormat = true;
+		BurnInNode->bOverride_BurnInClass = true;
+		BurnInNode->bOverride_bCompositeOntoFinalImage = true;
+		BurnInNode->OutputName = TEXT("Avid");
+		BurnInNode->FileNameFormat = BurnInFileNameFormat;
+		BurnInNode->OutputRestriction = FSoftClassPath(GetClass());
+		BurnInNode->LayerNameFormat = TEXT("");	// Unused
+		BurnInNode->BurnInClass = BurnInClass;
+		BurnInNode->bCompositeOntoFinalImage = bCompositeOntoFinalImage;
+
+		OutInjectedNodes.Add(BurnInNode);
+	}
+}
diff --git a/Engine/Plugins/Media/AvidDNxMedia/Source/Source/Private/MovieGraphAvidDNxHRNode.h b/Engine/Plugins/Media/AvidDNxMedia/Source/Source/Private/MovieGraphAvidDNxHRNode.h
--- a/Engine/Plugins/Media/AvidDNxMedia/Source/Source/Private/MovieGraphAvidDNxHRNode.h
+++ b/Engine/Plugins/Media/AvidDNxMedia/Source/Source/Private/MovieGraphAvidDNxHRNode.h
@@ -12,7 +12,7 @@ class FAvidDNxEncoder;
 
 /** A node which can output Avid DNxHR movies. */
 UCLASS(BlueprintType)
-class UMovieGraphAvidDNxHRNode : public UMovieGraphVideoOutputNode
+class UMovieGraphAvidDNxHRNode : public UMovieGraphVideoOutputNode, public IMovieGraphEvaluationNodeInjector
 {
 	GENERATED_BODY()
 
@@ -42,6 +42,10 @@ protected:
 	virtual void UpdateTelemetry(FMoviePipelineShotRenderTelemetry* InTelemetry) const override;
 	// ~UMovieGraphSettingNode Interface
 
+	// IMovieGraphEvaluationNodeInjector Interface
+	virtual void InjectNodesPostEvaluation(const FName& InBranchName, UMovieGraphEvaluatedConfig* InEvaluatedConfig, TArray<UMovieGraphSettingNode*>& OutInjectedNodes) override;
+	// ~IMovieGraphEvaluationNodeInjector Interface
+
 protected:
 	struct FAvidWriter : public MovieRenderGraph::IVideoCodecWriter
 	{
@@ -61,6 +65,18 @@ public:
 	UPROPERTY(EditAnywhere, BlueprintReadWrite, Category = Overrides, meta = (InlineEditConditionToggle))
 	uint8 bOverride_OCIOContext : 1;
 
+	UPROPERTY(EditAnywhere, BlueprintReadWrite, Category = Overrides, meta = (InlineEditConditionToggle))
+	uint8 bOverride_bEnableBurnIn : 1;
+
+	UPROPERTY(EditAnywhere, BlueprintReadWrite, Category = Overrides, meta = (InlineEditConditionToggle))
+	uint8 bOverride_BurnInClass : 1;
+	
+	UPROPERTY(EditAnywhere, BlueprintReadWrite, Category = Overrides, meta = (InlineEditConditionToggle))
+	uint8 bOverride_bCompositeOntoFinalImage : 1;
+
+	UPROPERTY(EditAnywhere, BlueprintReadWrite, Category = Overrides, meta = (InlineEditConditionToggle))
+	uint8 bOverride_BurnInFileNameFormat : 1;
+
 	/** The quality that the movie will be encoded with. */
 	UPROPERTY(EditAnywhere, BlueprintReadWrite, Category="Avid DNxHR", meta=(EditCondition="bOverride_Quality"))
 	EAvidDNxEncoderQuality Quality = EAvidDNxEncoderQuality::HQ_8bit;
@@ -85,4 +101,20 @@ public:
 	*/
 	UPROPERTY(EditAnywhere, BlueprintReadWrite, Category = "OCIO", DisplayName = "OCIO Context", meta = (DisplayAfter = "OCIOConfiguration", EditCondition = "bOverride_OCIOContext"))
 	TMap<FString, FString> OCIOContext;
+
+	/** If true, this output node will also generate a burn-in. */
+	UPROPERTY(EditAnywhere, BlueprintReadWrite, Category = "Burn In", meta = (EditCondition = "bOverride_bEnableBurnIn"))
+	bool bEnableBurnIn = false;
+
+	/** The widget that the burn-in should use. */
+	UPROPERTY(EditAnywhere, BlueprintReadWrite, Category="Burn In", meta=(MetaClass="/Script/MovieRenderPipelineCore.MovieGraphBurnInWidget", EditCondition="bOverride_BurnInClass"))
+	FSoftClassPath BurnInClass = FSoftClassPath(TEXT("/MovieRenderPipeline/Blueprints/Graph/DefaultGraphBurnIn.DefaultGraphBurnIn_C"));
+
+	/** If true, the burn-in that's generated will be composited onto this output. Otherwise, the burn-in will be written to a different file. */
+	UPROPERTY(EditAnywhere, BlueprintReadWrite, Category = "Burn In", meta = (EditCondition = "bOverride_bCompositeOntoFinalImage"))
+	bool bCompositeOntoFinalImage = true;
+
+	/** The file name format used for writing out the burn-in video (if not composited). */
+	UPROPERTY(EditAnywhere, BlueprintReadWrite, Category = "Burn In", meta = (EditCondition = "bOverride_BurnInFileNameFormat"))
+	FString BurnInFileNameFormat = TEXT("{sequence_name}.{layer_name}.{renderer_name}");
 };
\ No newline at end of file
diff --git a/Engine/Plugins/MovieScene/MovieRenderPipeline/Source/MovieRenderPipelineCore/Private/Graph/MovieGraphConfig.cpp b/Engine/Plugins/MovieScene/MovieRenderPipeline/Source/MovieRenderPipelineCore/Private/Graph/MovieGraphConfig.cpp
--- a/Engine/Plugins/MovieScene/MovieRenderPipeline/Source/MovieRenderPipelineCore/Private/Graph/MovieGraphConfig.cpp
+++ b/Engine/Plugins/MovieScene/MovieRenderPipeline/Source/MovieRenderPipelineCore/Private/Graph/MovieGraphConfig.cpp
@@ -2200,7 +2200,46 @@ UMovieGraphEvaluatedConfig* UMovieGraphConfig::CreateFlattenedGraph(const FMovie
 		}		
 	}
 
-	// Now that we have traversed the graph and flattened it, we need to give nodes a chance to perform a Post-Load type upgrade to the flattened graph.
+	// Give the IMovieGraphEvaluationNodeInjector interface a chance to inject nodes now that the evaluation has finished.
+	for (TPair<FName, FMovieGraphEvaluatedBranchConfig>& Pair : NewContext->BranchConfigMapping)
+	{
+		const FName& BranchName = Pair.Key;
+		FMovieGraphEvaluatedBranchConfig& BranchConfig = Pair.Value;
+
+		TArray<UMovieGraphSettingNode*> InjectedNodes;
+		
+		for (const TObjectPtr<UMovieGraphNode>& Node : BranchConfig.GetNodes())
+		{
+			if (IMovieGraphEvaluationNodeInjector* NodeInjector = Cast<IMovieGraphEvaluationNodeInjector>(Node))
+			{
+				NodeInjector->InjectNodesPostEvaluation(BranchName, NewContext, InjectedNodes);
+			}
+		}
+
+		// Add the injected nodes, but make sure that there isn't a node of the same instance name + type already in the branch.
+		for (UMovieGraphSettingNode* InjectedNode : InjectedNodes)
+		{
+			TArray<TObjectPtr<UMovieGraphNode>>& NodesWithInstanceName = BranchConfig.NamedNodes.FindOrAdd(InjectedNode->GetNodeInstanceName()).NodeInstances;
+
+			for (const TObjectPtr<UMovieGraphNode>& ExistingNode : NodesWithInstanceName)
+			{
+				if (ExistingNode && (ExistingNode->GetClass() == InjectedNode->GetClass()))
+				{
+					const FText InvalidNodeTypesError = FText::Format(LOCTEXT("TraversalError_SameTypeSameBranchConflict", "Found multiple nodes with type [{0}] and the same instance name [{1}] on the same branch [{2}]. Some may have been injected by other nodes. Aborting graph evaluation."),
+						{ FText::FromString(InjectedNode->GetClass()->GetName()), FText::FromString(InjectedNode->GetNodeInstanceName()), FText::FromName(BranchName) });
+					UE_LOG(LogMovieRenderPipeline, Error, TEXT("%s"), *InvalidNodeTypesError.ToString());
+					OutError = InvalidNodeTypesError.ToString();
+					return nullptr;
+				}
+			}
+			
+			NodesWithInstanceName.Add(InjectedNode);
+		}
+
+		InjectedNodes.Reset();
+	}
+
+	// We need to give nodes a chance to perform a Post-Load type upgrade to the flattened graph.
 	// This is because their scripting pipeline may have modified the source graph this flattened graph was created from to set legacy properties,
 	// so we need to give the flattened graphs a chance to migrate that data and warn the user of the need to update their scripts.
 	for (const TPair<FName, FMovieGraphEvaluatedBranchConfig>& Pair : NewContext->BranchConfigMapping)
@@ -2214,6 +2253,29 @@ UMovieGraphEvaluatedConfig* UMovieGraphConfig::CreateFlattenedGraph(const FMovie
 		}
 	}
 
+	// Check all nodes for validation errors. If any, these will be logged, and the render will be halted.
+	TArray<FText> ValidationErrors;
+	for (const TPair<FName, FMovieGraphEvaluatedBranchConfig>& Pair : NewContext->BranchConfigMapping)
+	{
+		for (UMovieGraphNode* Node : Pair.Value.GetNodes())
+		{
+			if (UMovieGraphSettingNode* Setting = Cast<UMovieGraphSettingNode>(Node))
+			{
+				Setting->GetNodeValidationErrors(Pair.Key, NewContext, ValidationErrors);
+			}
+		}
+	}
+
+	if (!ValidationErrors.IsEmpty())
+	{
+		const FText AllValidationErrors = FText::Join(INVTEXT("\n"), ValidationErrors);
+		const FText ValidationErrorMsg = FText::Format(LOCTEXT("TraversalError_NodeValidationErrors", "Found node validation errors while evaluating the graph.\n{0}"), AllValidationErrors);
+		
+		OutError = ValidationErrorMsg.ToString();
+		
+		return nullptr;
+	}
+
 	// Update the evaluated config with all of the evaluated variable values. These values are typically only used by scripting or blueprints.
 	for (const UMovieGraphConfig* EvaluatedGraph : NewContext->EvaluatedGraphs)
 	{
diff --git a/Engine/Plugins/MovieScene/MovieRenderPipeline/Source/MovieRenderPipelineCore/Private/Graph/MovieGraphDefaultRenderer.cpp b/Engine/Plugins/MovieScene/MovieRenderPipeline/Source/MovieRenderPipelineCore/Private/Graph/MovieGraphDefaultRenderer.cpp
--- a/Engine/Plugins/MovieScene/MovieRenderPipeline/Source/MovieRenderPipelineCore/Private/Graph/MovieGraphDefaultRenderer.cpp
+++ b/Engine/Plugins/MovieScene/MovieRenderPipeline/Source/MovieRenderPipelineCore/Private/Graph/MovieGraphDefaultRenderer.cpp
@@ -45,8 +45,11 @@ void UMovieGraphDefaultRenderer::SetupRenderingPipelineForShot(UMoviePipelineExe
 	{
 		TSubclassOf<UMovieGraphRenderPassNode> ClassType;
 
-		/** Maps a named branch to the specific render pass node that is assigned to render it. */
-		TMap<FMovieGraphRenderDataIdentifier, TWeakObjectPtr<UMovieGraphRenderPassNode>> BranchRenderers;
+		/**
+		 * Maps a named branch to the specific render pass nodes that are assigned to render it. Typically there is only one render pass node of a
+		 * given class type in the same branch; however, there are exceptions when node instance names are used (eg, burn-in nodes).
+		 */
+		TMap<FMovieGraphRenderDataIdentifier, TArray<TWeakObjectPtr<UMovieGraphRenderPassNode>>> BranchRenderers;
 	};
 
 	TArray<FMovieGraphPass> OutputPasses;
@@ -103,7 +106,6 @@ void UMovieGraphDefaultRenderer::SetupRenderingPipelineForShot(UMoviePipelineExe
 
 		for (UMovieGraphRenderPassNode* RenderPassNode : Renderers)
 		{
-
 			FMovieGraphPass* ExistingPass = OutputPasses.FindByPredicate([RenderPassNode](const FMovieGraphPass& Pass)
 				{ return Pass.ClassType == RenderPassNode->GetClass(); });
 			
@@ -119,7 +121,7 @@ void UMovieGraphDefaultRenderer::SetupRenderingPipelineForShot(UMoviePipelineExe
 
 			IdentifierToLayerIndex.Add(Identifier, RenderLayerIndex++);
 			
-			ExistingPass->BranchRenderers.Add(Identifier, RenderPassNode);
+			ExistingPass->BranchRenderers.FindOrAdd(Identifier).Add(RenderPassNode);
 		}
 	}
 
@@ -134,26 +136,28 @@ void UMovieGraphDefaultRenderer::SetupRenderingPipelineForShot(UMoviePipelineExe
 		
 		FMovieGraphRenderPassSetupData SetupData;
 		SetupData.Renderer = this;
-		//UE_LOG(LogMovieRenderPipeline, Warning, TEXT("\tRenderer Class: %s"), *Pass.ClassType->GetName());
-		for (const TTuple<FMovieGraphRenderDataIdentifier, TWeakObjectPtr<UMovieGraphRenderPassNode>>& BranchRenderer : Pass.BranchRenderers)
+		
+		for (const TTuple<FMovieGraphRenderDataIdentifier, TArray<TWeakObjectPtr<UMovieGraphRenderPassNode>>>& BranchRenderer : Pass.BranchRenderers)
 		{
-			for (int32 CameraIndex = 0; CameraIndex < CameraIndexes.Num(); CameraIndex++)
+			// Generate setup data for every render node / camera combination
+			for (const TWeakObjectPtr<UMovieGraphRenderPassNode>& RenderNode : BranchRenderer.Value)
 			{
-				FMovieGraphRenderPassLayerData& LayerData = SetupData.Layers.AddDefaulted_GetRef();
-				LayerData.BranchName = BranchRenderer.Key.RootBranchName;
-				LayerData.LayerName = BranchRenderer.Key.LayerName;
-				LayerData.LayerIndex = *IdentifierToLayerIndex.Find(BranchRenderer.Key);
-				LayerData.RenderPassNode = BranchRenderer.Value;
-				// Result could be -1 or 0...n
-				LayerData.CameraIndex = CameraIndexes[CameraIndex]; 
-				LayerData.NumCameras = CameraIndexes.Num();
+				for (int32 CameraIndex = 0; CameraIndex < CameraIndexes.Num(); CameraIndex++)
+				{
+					FMovieGraphRenderPassLayerData& LayerData = SetupData.Layers.AddDefaulted_GetRef();
+					LayerData.BranchName = BranchRenderer.Key.RootBranchName;
+					LayerData.LayerName = BranchRenderer.Key.LayerName;
+					LayerData.LayerIndex = *IdentifierToLayerIndex.Find(BranchRenderer.Key);
+					LayerData.RenderPassNode = RenderNode;
+					// Result could be -1 or 0...n
+					LayerData.CameraIndex = CameraIndexes[CameraIndex]; 
+					LayerData.NumCameras = CameraIndexes.Num();
 				
-				// We provide the name here so that the setup functions of the renderers can create RenderLayerIdentifiers.
-				UE::MovieGraph::DefaultRenderer::FCameraInfo CameraInfo = GetCameraInfo(LayerData.CameraIndex);
-				LayerData.CameraName = CameraInfo.CameraName;
+					// We provide the name here so that the setup functions of the renderers can create RenderLayerIdentifiers.
+					UE::MovieGraph::DefaultRenderer::FCameraInfo CameraInfo = GetCameraInfo(LayerData.CameraIndex);
+					LayerData.CameraName = CameraInfo.CameraName;
+				}
 			}
-				
-			// UE_LOG(LogMovieRenderPipeline, Warning, TEXT("\t\tBranch Name: %s"), *LayerBranchName.ToString());
 		}
 
 		UMovieGraphRenderPassNode* RenderPassCDO = Pass.ClassType->GetDefaultObject<UMovieGraphRenderPassNode>();
diff --git a/Engine/Plugins/MovieScene/MovieRenderPipeline/Source/MovieRenderPipelineCore/Private/Graph/MovieGraphNode.cpp b/Engine/Plugins/MovieScene/MovieRenderPipeline/Source/MovieRenderPipelineCore/Private/Graph/MovieGraphNode.cpp
--- a/Engine/Plugins/MovieScene/MovieRenderPipeline/Source/MovieRenderPipelineCore/Private/Graph/MovieGraphNode.cpp
+++ b/Engine/Plugins/MovieScene/MovieRenderPipeline/Source/MovieRenderPipelineCore/Private/Graph/MovieGraphNode.cpp
@@ -432,6 +432,12 @@ bool UMovieGraphNode::IsDisabled() const
 	return bIsDisabled;
 }
 
+bool UMovieGraphNode::GetNodeValidationErrors(const FName& InBranchName, const UMovieGraphEvaluatedConfig* InEvaluatedConfig, TArray<FText>& OutValidationErrors) const
+{
+	// No validation errors by default
+	return false;
+}
+
 #if WITH_EDITOR
 FLinearColor UMovieGraphNode::GetNodeTitleColor() const
 {
@@ -640,3 +646,8 @@ TArray<FMovieGraphPinProperties> UMovieGraphSettingNode::GetOutputPinProperties(
 	Properties.Add(FMovieGraphPinProperties::MakeBranchProperties());
 	return Properties;
 }
+
+void IMovieGraphEvaluationNodeInjector::InjectNodesPostEvaluation(const FName& InBranchName, UMovieGraphEvaluatedConfig* InEvaluatedConfig, TArray<UMovieGraphSettingNode*>& OutInjectedNodes)
+{
+	
+}
\ No newline at end of file
diff --git a/Engine/Plugins/MovieScene/MovieRenderPipeline/Source/MovieRenderPipelineCore/Private/Graph/MovieGraphPipeline.cpp b/Engine/Plugins/MovieScene/MovieRenderPipeline/Source/MovieRenderPipelineCore/Private/Graph/MovieGraphPipeline.cpp
--- a/Engine/Plugins/MovieScene/MovieRenderPipeline/Source/MovieRenderPipelineCore/Private/Graph/MovieGraphPipeline.cpp
+++ b/Engine/Plugins/MovieScene/MovieRenderPipeline/Source/MovieRenderPipelineCore/Private/Graph/MovieGraphPipeline.cpp
@@ -117,7 +117,7 @@ void UMovieGraphPipeline::Initialize(UMoviePipelineExecutorJob* InJob, const FMo
 	else
 	{
 		// Although the evaluated graph was created for Unreal Insights, it will lead to a general failure to initialize the pipeline.
-		UE_LOG(LogMovieRenderPipeline, Error, TEXT("Unable to initialize Movie Pipeline. The following error was encountered: %s"), *OutError);
+		UE_LOG(LogMovieRenderPipeline, Error, TEXT("Unable to initialize Movie Pipeline. The following error(s) were encountered: %s"), *OutError);
 	}
 
 	// Create instances of our different classes from the InitConfig
diff --git a/Engine/Plugins/MovieScene/MovieRenderPipeline/Source/MovieRenderPipelineCore/Private/Graph/Nodes/MovieGraphBurnInNode.cpp b/Engine/Plugins/MovieScene/MovieRenderPipeline/Source/MovieRenderPipelineCore/Private/Graph/Nodes/MovieGraphBurnInNode.cpp
--- a/Engine/Plugins/MovieScene/MovieRenderPipeline/Source/MovieRenderPipelineCore/Private/Graph/Nodes/MovieGraphBurnInNode.cpp
+++ b/Engine/Plugins/MovieScene/MovieRenderPipeline/Source/MovieRenderPipelineCore/Private/Graph/Nodes/MovieGraphBurnInNode.cpp
@@ -13,6 +13,8 @@
 
 #include UE_INLINE_GENERATED_CPP_BY_NAME(MovieGraphBurnInNode)
 
+#define LOCTEXT_NAMESPACE "MovieGraphNodes"
+
 const FString UMovieGraphBurnInNode::RendererName = FString("BurnIn");
 const FString UMovieGraphBurnInNode::DefaultBurnInWidgetAsset = TEXT("/MovieRenderPipeline/Blueprints/Graph/DefaultGraphBurnIn.DefaultGraphBurnIn_C");
 
@@ -24,7 +26,7 @@ UMovieGraphBurnInNode::UMovieGraphBurnInNode()
 #if WITH_EDITOR
 FText UMovieGraphBurnInNode::GetNodeTitle(const bool bGetDescriptive) const
 {
-	return NSLOCTEXT("MovieGraphNodes", "BurnInGraphNode_Description", "Burn In");
+	return LOCTEXT("BurnInGraphNode_Description", "Burn In [DEPRECATED]\nUse burn-in settings on output nodes.");
 }
 
 FLinearColor UMovieGraphBurnInNode::GetNodeTitleColor() const
@@ -65,6 +67,68 @@ void UMovieGraphBurnInNode::TeardownImpl()
 	BurnInWidgetInstances.Empty();
 }
 
+FString UMovieGraphOutputBurnInNode::GetNodeInstanceName() const
+{
+	return OutputName;
+}
+
+#if WITH_EDITOR
+FText UMovieGraphOutputBurnInNode::GetNodeTitle(const bool bGetDescriptive) const
+{
+	return LOCTEXT("BurnInOutputGraphNode_Description", "Burn In");
+}
+#endif	// WITH_EDITOR
+
+TArray<FSoftClassPath> UMovieGraphOutputBurnInNode::GetOutputTypeRestrictionsImpl() const
+{
+	return { OutputRestriction };
+}
+
+FString UMovieGraphOutputBurnInNode::GetRendererNameImpl() const
+{
+	return FString::Format(TEXT("BurnIn_{0}"), { OutputName });
+}
+
+FString UMovieGraphOutputBurnInNode::GetFileNameFormatOverride() const
+{
+	return FileNameFormat;
+}
+
+FString UMovieGraphOutputBurnInNode::GetLayerNameOverride() const
+{
+	return LayerNameFormat;
+}
+
+bool UMovieGraphOutputBurnInNode::GetNodeValidationErrors(const FName& InBranchName, const UMovieGraphEvaluatedConfig* InEvaluatedConfig, TArray<FText>& OutValidationErrors) const
+{
+	if (Super::GetNodeValidationErrors(InBranchName, InEvaluatedConfig, OutValidationErrors))
+	{
+		return false;
+	}
+
+	// The old burn-in node (deprecated) cannot be used with the output-type-specific burn-in node (this node). If the old node appears in globals,
+	// or appears in the same branch as this node, that is an error and should halt the render.
+	for (const FName& BranchName : InEvaluatedConfig->GetBranchNames())
+	{
+		constexpr bool bIncludeCDOs = false;
+		constexpr bool bExactMatch = true;
+		if (InEvaluatedConfig->GetSettingForBranch(UMovieGraphBurnInNode::StaticClass(), BranchName, bIncludeCDOs, bExactMatch))
+		{
+			if ((BranchName == GlobalsPinName) || (BranchName == InBranchName))
+			{
+				const FText OutputClassName = OutputRestriction.ResolveClass()
+					? FText::FromString(OutputRestriction.ResolveClass()->GetName())
+					: LOCTEXT("BurnInOutputNode_UnknownClass", "Unknown");
+
+				OutValidationErrors.Add(FText::Format(LOCTEXT("BurnInOutputNode_Incompatibility", "An output-specific ({0}) burn-in was added to a graph that also uses the deprecated Burn In node. These two nodes are incompatible; only the Burn In node or the output-specific burn-ins can be used."), OutputClassName));
+				return true;
+			}
+		}
+	}
+
+	return false;
+}
+
 TObjectPtr<UMovieGraphBurnInWidget> UMovieGraphBurnInNode::GetOrCreateBurnInWidget(UClass* InWidgetClass, UWorld* InOwner)
 {
 	if (const TObjectPtr<UMovieGraphBurnInWidget>* ExistingBurnInWidget = BurnInWidgetInstances.Find(InWidgetClass))
@@ -160,3 +224,5 @@ UClass* UMovieGraphBurnInNode::FMovieGraphBurnInPass::GetBurnInClass() const
 {
 	return CachedBurnInWidgetClass.Get();
 }
+
+#undef LOCTEXT_NAMESPACE // "MovieGraphNodes"
\ No newline at end of file
diff --git a/Engine/Plugins/MovieScene/MovieRenderPipeline/Source/MovieRenderPipelineCore/Private/Graph/Nodes/MovieGraphVideoOutputNode.cpp b/Engine/Plugins/MovieScene/MovieRenderPipeline/Source/MovieRenderPipelineCore/Private/Graph/Nodes/MovieGraphVideoOutputNode.cpp
--- a/Engine/Plugins/MovieScene/MovieRenderPipeline/Source/MovieRenderPipelineCore/Private/Graph/Nodes/MovieGraphVideoOutputNode.cpp
+++ b/Engine/Plugins/MovieScene/MovieRenderPipeline/Source/MovieRenderPipelineCore/Private/Graph/Nodes/MovieGraphVideoOutputNode.cpp
@@ -67,6 +67,15 @@ void UMovieGraphVideoOutputNode::OnReceiveImageDataImpl(UMovieGraphPipeline* InP
 
 	for (FMovieGraphPassData& RenderPassData : InRawFrameData->ImageOutputData)
 	{
+		FImagePixelData* RawRenderPassData = RenderPassData.Value.Get();
+		const UE::MovieGraph::FMovieGraphSampleState* Payload = RawRenderPassData->GetPayload<UE::MovieGraph::FMovieGraphSampleState>();
+
+		// This payload may have opted out of being written by this output type.
+		if (!Payload->OutputTypeRestrictions.IsEmpty() && !Payload->OutputTypeRestrictions.Contains(FSoftClassPath(GetClass())))
+		{
+			continue;
+		}
+		
 		// Don't write out a composited pass in this loop, as it will be composited by the encoder and not written separately. 
 		if (CompositedPasses.ContainsByPredicate([&RenderPassData](const FMovieGraphPassData& CompositedPass)
 			{
@@ -89,9 +98,6 @@ void UMovieGraphVideoOutputNode::OnReceiveImageDataImpl(UMovieGraphPipeline* InP
 			continue;
 		}
 
-		FImagePixelData* RawRenderPassData = RenderPassData.Value.Get();
-		const UE::MovieGraph::FMovieGraphSampleState* Payload = RawRenderPassData->GetPayload<UE::MovieGraph::FMovieGraphSampleState>();
-
 		// We can have multiple shots in one file, so we find/add the shot in case this is a new shot, and the counter will start over at zero.
 		MovieRenderGraph::IVideoCodecWriter::FLightweightSourceData& LightweightData = OutputWriter->CodecWriter->LightweightSourceData.FindOrAdd(Payload->TraversalContext.ShotIndex);
 		LightweightData.SubmittedFrameCount++;
@@ -144,6 +150,14 @@ void UMovieGraphVideoOutputNode::OnReceiveImageDataImpl(UMovieGraphPipeline* InP
 					continue;
 				}
 
+				UE::MovieGraph::FMovieGraphSampleState* CompositePayload = CompositePass.Value->GetPayload<UE::MovieGraph::FMovieGraphSampleState>();
+
+				// This composite may have opted out of being written by this output type.
+				if (!CompositePayload->OutputTypeRestrictions.IsEmpty() && !CompositePayload->OutputTypeRestrictions.Contains(FSoftClassPath(GetClass())))
+				{
+					continue;
+				}
+
 				// There might be composites across multiple branches; pick out the composites specific to this branch.
 				if (RenderPassData.Key.RootBranchName != CompositePass.Key.RootBranchName)
 				{
@@ -255,6 +269,11 @@ void UMovieGraphVideoOutputNode::GetOutputFilePaths(const UMovieGraphPipeline* I
 	UMovieGraphVideoOutputNode* EvaluatedNode = Cast<UMovieGraphVideoOutputNode>(
 		EvaluatedConfig->GetSettingForBranch(GetClass(), InRenderPassData.Key.RootBranchName, bIncludeCDOs, bExactMatch));
 
+	const UE::MovieGraph::FMovieGraphSampleState* Payload = InRenderPassData.Value->GetPayload<UE::MovieGraph::FMovieGraphSampleState>();
+
+	// The file name format usually comes from the output node directly, but the payload has a chance to override it.
+	FString FileNameFormatString = !Payload->FilenameFormatOverride.IsEmpty() ? Payload->FilenameFormatOverride : EvaluatedNode->FileNameFormat;
+
 	const TMap<FString, FString> FileNameFormatOverrides =  {
 		{TEXT("camera_name"), InRenderPassData.Key.CameraName},
 		{TEXT("render_pass"), InRenderPassData.Key.RendererName},
@@ -262,11 +281,10 @@ void UMovieGraphVideoOutputNode::GetOutputFilePaths(const UMovieGraphPipeline* I
 	};
 	FMovieGraphFilenameResolveParams ResolveParams =
 		FMovieGraphFilenameResolveParams::MakeResolveParams(InRenderPassData.Key, InPipeline, EvaluatedConfig, InRawFrameData->TraversalContext, FileNameFormatOverrides);
-	ResolveParams.FileNameOverride = EvaluatedNode->FileNameFormat;
+	ResolveParams.FileNameOverride = FileNameFormatString;
 	ResolveParams.Version = Shot ? Shot->ShotInfo.VersionNumber : UMovieGraphBlueprintLibrary::ResolveVersionNumber(ResolveParams);
 	
 	FMovieGraphResolveArgs FinalFormatArgs;
-	FString FileNameFormatString = EvaluatedNode->FileNameFormat;
 	
 	// Insert tokens like {layer_name} as appropriate to make sure outputs don't clash with each other.
 	DisambiguateFilename(FileNameFormatString, InRawFrameData, EvaluatedNode->GetFName(), InRenderPassData);
diff --git a/Engine/Plugins/MovieScene/MovieRenderPipeline/Source/MovieRenderPipelineCore/Private/Graph/Nodes/MovieGraphWidgetRendererBaseNode.cpp b/Engine/Plugins/MovieScene/MovieRenderPipeline/Source/MovieRenderPipelineCore/Private/Graph/Nodes/MovieGraphWidgetRendererBaseNode.cpp
--- a/Engine/Plugins/MovieScene/MovieRenderPipeline/Source/MovieRenderPipelineCore/Private/Graph/Nodes/MovieGraphWidgetRendererBaseNode.cpp
+++ b/Engine/Plugins/MovieScene/MovieRenderPipeline/Source/MovieRenderPipelineCore/Private/Graph/Nodes/MovieGraphWidgetRendererBaseNode.cpp
@@ -19,6 +19,18 @@
 
 #include UE_INLINE_GENERATED_CPP_BY_NAME(MovieGraphWidgetRendererBaseNode)
 
+FString UMovieGraphWidgetRendererBaseNode::GetFileNameFormatOverride() const
+{
+	// By default, there's no override
+	return FString();
+}
+
+FString UMovieGraphWidgetRendererBaseNode::GetLayerNameOverride() const
+{
+	// By default, there's no override
+	return FString();
+}
+
 void UMovieGraphWidgetRendererBaseNode::SetupImpl(const FMovieGraphRenderPassSetupData& InSetupData)
 {
 	for (const FMovieGraphRenderPassLayerData& LayerData : InSetupData.Layers)
@@ -107,8 +119,7 @@ UMovieGraphWidgetRendererBaseNode* UMovieGraphWidgetRendererBaseNode::FMovieGrap
 void UMovieGraphWidgetRendererBaseNode::FMovieGraphWidgetPass::Render(const FMovieGraphTraversalContext& InFrameTraversalContext, const FMovieGraphTimeStepData& InTimeData)
 {
 	const UMovieGraphPipeline* Pipeline = Renderer->GetOwningGraph();
-	UMovieGraphWidgetRendererBaseNode* ParentNodeThisFrame = GetParentNode(InTimeData.EvaluatedConfig);
-
+	UMovieGraphWidgetRendererBaseNode* ParentNodeThisFrame = Cast<UMovieGraphWidgetRendererBaseNode>(LayerData.RenderPassNode);
 
 	if(Pipeline->GetActiveShotList()[Pipeline->GetCurrentShotIndex()]->ShotInfo.State != EMovieRenderShotState::Rendering)
 	{
@@ -168,10 +179,23 @@ void UMovieGraphWidgetRendererBaseNode::FMovieGraphWidgetPass::Render(const FMov
 			1.f, OutputResolution, InTimeData.FrameDeltaTime);
 
 		ENQUEUE_RENDER_COMMAND(WidgetRenderTargetResolveCommand)(
-			[InTimeData, InFrameTraversalContext, RenderTargetInitParams, RenderDataIdentifier = this->RenderDataIdentifier,
-			bComposite = ParentNodeThisFrame->bCompositeOntoFinalImage, CompositingSortOrder = GetCompositingSortOrder(),
-			BackbufferRenderTarget, OutputMerger, OverscannedResolution, CropRect, OutputResolution,
-			RenderLayerIndex = this->LayerData.LayerIndex](FRHICommandListImmediate& RHICmdList)
+			[
+				InTimeData,
+				InFrameTraversalContext,
+				RenderTargetInitParams,
+				RenderDataIdentifier = this->RenderDataIdentifier,
+				bComposite = ParentNodeThisFrame->bCompositeOntoFinalImage,
+				CompositingSortOrder = GetCompositingSortOrder(),
+				BackbufferRenderTarget,
+				OutputMerger,
+				OverscannedResolution,
+				CropRect,
+				OutputResolution,
+				RenderLayerIndex = this->LayerData.LayerIndex,
+				OutputTypeRestrictions = ParentNodeThisFrame->GetOutputTypeRestrictions(),
+				FileNameFormatOverride = ParentNodeThisFrame->GetFileNameFormatOverride(),
+				LayerNameFormatOverride = ParentNodeThisFrame->GetLayerNameOverride()
+			](FRHICommandListImmediate& RHICmdList)
 			{
 				const FIntRect SourceRect = FIntRect(0, 0, BackbufferRenderTarget->GetSizeXY().X, BackbufferRenderTarget->GetSizeXY().Y);
 
@@ -204,6 +228,9 @@ void UMovieGraphWidgetRendererBaseNode::FMovieGraphWidgetPass::Render(const FMov
 				SampleStatePayload->bCompositeOnOtherRenders = bComposite;
 				SampleStatePayload->CompositingSortOrder = CompositingSortOrder;
 				SampleStatePayload->RenderLayerIndex = RenderLayerIndex;
+				SampleStatePayload->OutputTypeRestrictions = OutputTypeRestrictions;
+				SampleStatePayload->FilenameFormatOverride = FileNameFormatOverride;
+				SampleStatePayload->LayerNameOverride = LayerNameFormatOverride;
 
 				TUniquePtr<FImagePixelData> PixelData =
 					MakeUnique<TImagePixelData<FColor>>(OutputResolution, TArray64<FColor>(MoveTemp(RawPixels)), SampleStatePayload);
diff --git a/Engine/Plugins/MovieScene/MovieRenderPipeline/Source/MovieRenderPipelineCore/Public/Graph/MovieGraphDataTypes.h b/Engine/Plugins/MovieScene/MovieRenderPipeline/Source/MovieRenderPipelineCore/Public/Graph/MovieGraphDataTypes.h
--- a/Engine/Plugins/MovieScene/MovieRenderPipeline/Source/MovieRenderPipelineCore/Public/Graph/MovieGraphDataTypes.h
+++ b/Engine/Plugins/MovieScene/MovieRenderPipeline/Source/MovieRenderPipelineCore/Public/Graph/MovieGraphDataTypes.h
@@ -441,6 +441,9 @@ namespace UE::MovieGraph
 		*/
 		FString Debug_OverrideFilename;
 
+		/** Specifies the filename format that this sample should use when being written to disk. Not used for debug like Debug_OverrideFilename. */
+		FString FilenameFormatOverride;
+
 		/** Set this to true for every sample if an accumulator is required (ie: has temporal/spatial sub-samples, or high res tiling). */
 		bool bRequiresAccumulator;
 
@@ -502,6 +505,12 @@ namespace UE::MovieGraph
 		 * exact layer name may be needed.
 		 */
 		FString LayerNameOverride;
+
+		/**
+		 * Specifies the output node class(es) that can output this sample. If this array is empty, this sample will be processed by all
+		 * output nodes (which is typical).
+		 */
+		TArray<FSoftClassPath> OutputTypeRestrictions;
 	};
 
 	/**
diff --git a/Engine/Plugins/MovieScene/MovieRenderPipeline/Source/MovieRenderPipelineCore/Public/Graph/MovieGraphNode.h b/Engine/Plugins/MovieScene/MovieRenderPipeline/Source/MovieRenderPipelineCore/Public/Graph/MovieGraphNode.h
--- a/Engine/Plugins/MovieScene/MovieRenderPipeline/Source/MovieRenderPipelineCore/Public/Graph/MovieGraphNode.h
+++ b/Engine/Plugins/MovieScene/MovieRenderPipeline/Source/MovieRenderPipelineCore/Public/Graph/MovieGraphNode.h
@@ -304,6 +304,13 @@ public:
 	UFUNCTION(BlueprintCallable, Category = "Movie Graph")
 	MOVIERENDERPIPELINECORE_API bool IsDisabled() const;
 
+	/**
+	 * Gets any validation errors that the node generated. Validation errors will halt the active render, if any. Called on the fully-evaluated
+	 * node post-evaluation. Returns true if there were any errors (added to OutValidationErrors), else false.
+	 */
+	UFUNCTION(BlueprintCallable, Category = "Movie Graph")
+	MOVIERENDERPIPELINECORE_API virtual bool GetNodeValidationErrors(const FName& InBranchName, const UMovieGraphEvaluatedConfig* InEvaluatedConfig, TArray<FText>& OutValidationErrors) const;
+
 #if WITH_EDITOR
 	UFUNCTION(BlueprintPure, Category = "Movie Graph")
 	int32 GetNodePosX() const { return NodePosX; }
@@ -510,3 +517,28 @@ public:
 	/** Returns true if this node has finished its export process, else false. */
 	virtual bool HasFinishedExporting() = 0;
 };
+
+UINTERFACE(MinimalAPI, NotBlueprintable)
+class UMovieGraphEvaluationNodeInjector : public UInterface
+{
+	GENERATED_BODY()
+};
+
+/**
+ * Provides the interface that nodes must implement if they need to inject additional nodes into the graph at various points in the
+ * graph evaluation cycle.
+ */
+class IMovieGraphEvaluationNodeInjector : public IInterface
+{
+	GENERATED_BODY()
+
+public:
+	/**
+	 * Injects nodes into the EVALUATED graph. Called on the fully evaluated node. All new nodes should be outered to InEvaluatedConfig. Because these
+	 * nodes are injected post-evaluation, the movie pipeline will be immediately shut down if any of the returned nodes already exist in the branch
+	 * (having two nodes of the same type in the same branch post-evaluation is not allowed). Note that modifying the Globals branch is not typically
+	 * advised because at this stage, it has already been merged into all other branches.
+	 */
+	UFUNCTION(BlueprintCallable, Category = "Evaluation")
+	MOVIERENDERPIPELINECORE_API virtual void InjectNodesPostEvaluation(const FName& InBranchName, UMovieGraphEvaluatedConfig* InEvaluatedConfig, TArray<UMovieGraphSettingNode*>& OutInjectedNodes);
+};
diff --git a/Engine/Plugins/MovieScene/MovieRenderPipeline/Source/MovieRenderPipelineCore/Public/Graph/Nodes/MovieGraphBurnInNode.h b/Engine/Plugins/MovieScene/MovieRenderPipeline/Source/MovieRenderPipelineCore/Public/Graph/Nodes/MovieGraphBurnInNode.h
--- a/Engine/Plugins/MovieScene/MovieRenderPipeline/Source/MovieRenderPipelineCore/Public/Graph/Nodes/MovieGraphBurnInNode.h
+++ b/Engine/Plugins/MovieScene/MovieRenderPipeline/Source/MovieRenderPipelineCore/Public/Graph/Nodes/MovieGraphBurnInNode.h
@@ -14,8 +14,12 @@ class UMovieGraphBurnInWidget;
 class UMovieGraphDefaultRenderer;
 struct FMovieGraphRenderPassLayerData;
 
-/** A node which generates a widget burn-in, rendered to a standalone image or composited on top of a render layer. */
-UCLASS(MinimalAPI)
+/**
+ * A node which generates a widget burn-in, rendered to a standalone image or composited on top of a render layer.
+ *
+ * DEPRECATED. Use the burn-in settings on output nodes instead to generate burn-ins.
+ */
+UCLASS(MinimalAPI, Hidden, HideDropdown)
 class UMovieGraphBurnInNode : public UMovieGraphWidgetRendererBaseNode
 {
 	GENERATED_BODY()
@@ -82,4 +86,72 @@ protected:
 	// ~UMovieGraphRenderPassNode Interface
 };
 
+/**
+ * A node which generates a widget burn-in, rendered to a standalone image or composited on top of a render layer.
+ *
+ * This node is not meant to be added by the user directly. Instead, it's injected into the graph by a specific output node (eg, JPG). It can control
+ * which output type(s) write it out (see the OutputRestriction property).
+ * 
+ */
+UCLASS(MinimalAPI, Hidden, HideDropdown)
+class UMovieGraphOutputBurnInNode : public UMovieGraphBurnInNode
+{
+	GENERATED_BODY()
+
+public:
+	UE_API UMovieGraphOutputBurnInNode() = default;
+
+	// UMovieGraphSettingNode Interface
+	virtual FString GetNodeInstanceName() const override;
+	// ~UMovieGraphSettingNode Interface
+
+	// UMovieGraphNode Interface
+#if WITH_EDITOR
+	virtual FText GetNodeTitle(const bool bGetDescriptive) const override;
+#endif
+	virtual bool GetNodeValidationErrors(const FName& InBranchName, const UMovieGraphEvaluatedConfig* InEvaluatedConfig, TArray<FText>& OutValidationErrors) const override;
+	// ~UMovieGraphNode Interface
+
+protected:
+	virtual TArray<FSoftClassPath> GetOutputTypeRestrictionsImpl() const override;
+	virtual FString GetRendererNameImpl() const override;
+
+	// UMovieGraphWidgetRendererBaseNode Interface
+	UE_API virtual FString GetFileNameFormatOverride() const override;
+	UE_API virtual FString GetLayerNameOverride() const override;
+	// ~UMovieGraphWidgetRendererBaseNode Interface
+
+public:
+	UPROPERTY(EditAnywhere, BlueprintReadWrite, Category = Overrides, meta = (InlineEditConditionToggle))
+	uint8 bOverride_OutputName : 1;
+
+	UPROPERTY(EditAnywhere, BlueprintReadWrite, Category = Overrides, meta = (InlineEditConditionToggle))
+	uint8 bOverride_FileNameFormat : 1;
+
+	UPROPERTY(EditAnywhere, BlueprintReadWrite, Category = Overrides, meta = (InlineEditConditionToggle))
+	uint8 bOverride_OutputRestriction : 1;
+
+	UPROPERTY(EditAnywhere, BlueprintReadWrite, Category = Overrides, meta = (InlineEditConditionToggle))
+	uint8 bOverride_LayerNameFormat : 1;
+	
+	/**
+	 * The name of the output that is associated with this burn-in (eg, "JPG"). Will be used as the instance name for the node, so it should be
+	 * unique if multiple burn-in nodes need to co-exist in the same branch.
+	 */
+	UPROPERTY(EditAnywhere, BlueprintReadWrite, Category = "Burn In", meta=(EditCondition="bOverride_OutputName"))
+	FString OutputName;
+
+	/** The type(s) of nodes that can write out this burn-in. If no output types are specified, all active output types will write out the burn-in. */
+	UPROPERTY(EditAnywhere, BlueprintReadWrite, Category = "Burn In", meta=(EditCondition="bOverride_OutputRestriction"))
+	FSoftClassPath OutputRestriction;
+
+	/** If not composited, this is the format string used for the file name that the burn-in is written to. */
+	UPROPERTY(EditAnywhere, BlueprintReadWrite, Category = "Burn In", meta=(EditCondition="bOverride_FileNameFormat"))
+	FString FileNameFormat = "{sequence_name}.{layer_name}.{frame_number}";
+
+	/** For multi-layer output formats, this is the format string used to generate the layer name. */
+	UPROPERTY(EditAnywhere, BlueprintReadWrite, Category = "Burn In", meta=(EditCondition="bOverride_LayerNameFormat"))
+	FString LayerNameFormat;
+};
+
 #undef UE_API
diff --git a/Engine/Plugins/MovieScene/MovieRenderPipeline/Source/MovieRenderPipelineCore/Public/Graph/Nodes/MovieGraphRenderPassNode.h b/Engine/Plugins/MovieScene/MovieRenderPipeline/Source/MovieRenderPipelineCore/Public/Graph/Nodes/MovieGraphRenderPassNode.h
--- a/Engine/Plugins/MovieScene/MovieRenderPipeline/Source/MovieRenderPipelineCore/Public/Graph/Nodes/MovieGraphRenderPassNode.h
+++ b/Engine/Plugins/MovieScene/MovieRenderPipeline/Source/MovieRenderPipelineCore/Public/Graph/Nodes/MovieGraphRenderPassNode.h
@@ -57,6 +57,12 @@ public:
 	**/
 	virtual int32 GetCoolingDownFrameCount() const { return 0; }
 
+	/**
+	 * Gets the type(s) of output nodes that are allowed to process/output the pixel data originating from this node. Typically the return value is
+	 * an empty array, meaning that all output nodes process the data.
+	 */
+	TArray<FSoftClassPath> GetOutputTypeRestrictions() const { return GetOutputTypeRestrictionsImpl(); }
+
 #if WITH_EDITOR
 	virtual FText GetMenuCategory() const override
 	{
@@ -83,4 +89,5 @@ protected:
 	virtual void RenderImpl(const FMovieGraphTraversalContext& InFrameTraversalContext, const FMovieGraphTimeStepData& InTimeData) {}
 	virtual void GatherOutputPassesImpl(UMovieGraphEvaluatedConfig* InConfig, TArray<FMovieGraphRenderDataIdentifier>& OutExpectedPasses) const {}
 	virtual int32 GetNumSceneViewsRenderedImpl() const { return 0; }
+	virtual TArray<FSoftClassPath> GetOutputTypeRestrictionsImpl() const { return {}; } 
 };
diff --git a/Engine/Plugins/MovieScene/MovieRenderPipeline/Source/MovieRenderPipelineCore/Public/Graph/Nodes/MovieGraphWidgetRendererBaseNode.h b/Engine/Plugins/MovieScene/MovieRenderPipeline/Source/MovieRenderPipelineCore/Public/Graph/Nodes/MovieGraphWidgetRendererBaseNode.h
--- a/Engine/Plugins/MovieScene/MovieRenderPipeline/Source/MovieRenderPipelineCore/Public/Graph/Nodes/MovieGraphWidgetRendererBaseNode.h
+++ b/Engine/Plugins/MovieScene/MovieRenderPipeline/Source/MovieRenderPipelineCore/Public/Graph/Nodes/MovieGraphWidgetRendererBaseNode.h
@@ -46,7 +46,6 @@ protected:
 		virtual int32 GetCompositingSortOrder() const = 0;
 		virtual UMovieGraphWidgetRendererBaseNode* GetParentNode(UMovieGraphEvaluatedConfig* InConfig) const;
 
-
 	protected:
 		FMovieGraphRenderPassLayerData LayerData;
 
@@ -66,6 +65,12 @@ protected:
 protected:
 	/** Generate a new pass that the widget renderer will use. */
 	virtual TUniquePtr<FMovieGraphWidgetPass> GeneratePass() PURE_VIRTUAL(UMovieGraphWidgetRendererBaseNode::GeneratePass, return nullptr;);
+
+	/** Gets the file name that the render will write to if not performing a composite. */
+	UE_API virtual FString GetFileNameFormatOverride() const;
+
+	/** Gets the layer name this render will use when writing to multi-layer formats. An empty string means no override, and the default will be used . */
+	UE_API virtual FString GetLayerNameOverride() const;
 	
 	// UMovieGraphRenderPassNode Interface
 	UE_API virtual void SetupImpl(const FMovieGraphRenderPassSetupData& InSetupData) override;
diff --git a/Engine/Plugins/MovieScene/MovieRenderPipeline/Source/MovieRenderPipelineCore/Public/MoviePipelineUtils.h b/Engine/Plugins/MovieScene/MovieRenderPipeline/Source/MovieRenderPipelineCore/Public/MoviePipelineUtils.h
--- a/Engine/Plugins/MovieScene/MovieRenderPipeline/Source/MovieRenderPipelineCore/Public/MoviePipelineUtils.h
+++ b/Engine/Plugins/MovieScene/MovieRenderPipeline/Source/MovieRenderPipelineCore/Public/MoviePipelineUtils.h
@@ -166,25 +166,25 @@ MOVIERENDERPIPELINECORE_API void GetPassCompositeData(FMoviePipelineMergerOutput
  * this specific message within the current SCOPE (identifiers are NOT unique across scopes, so the same log message will be logged more than once per
  * render if it appears in different scopes).
  */
-#define MRG_LOG_ONCE_PER_RENDER(Identifier, LogCategory, LogVerbosity, LogFormat, ...)							\
-	static UWorld* PREPROCESSOR_JOIN(__LastRenderWorld__, Identifier) = nullptr;								\
-	static TArray<FString> PREPROCESSOR_JOIN(__LoggedMessages__, Identifier);									\
-																												\
-	/* If this world has not been encountered yet (ie, no renders yet, or this is a new render), clear out */	\
-	/* the messages that have been recorded as already logged. */												\
-	if (PREPROCESSOR_JOIN(__LastRenderWorld__, Identifier) != MoviePipeline::FindCurrentWorld())				\
-	{																											\
-		PREPROCESSOR_JOIN(__LastRenderWorld__, Identifier) = MoviePipeline::FindCurrentWorld();					\
-		PREPROCESSOR_JOIN(__LoggedMessages__, Identifier).Empty();												\
-	}																											\
-																												\
-	/* If this message hasn't been logged, record that it has been logged so we can skip it for other frames. */\
-	const FString FinalLogMessage = FString::Printf(LogFormat, ##__VA_ARGS__);									\
-	if (!PREPROCESSOR_JOIN(__LoggedMessages__, Identifier).Contains(FinalLogMessage))							\
-	{																											\
-		UE_LOG(LogCategory, LogVerbosity, TEXT("%s"), *FinalLogMessage);										\
-		PREPROCESSOR_JOIN(__LoggedMessages__, Identifier).Add(FinalLogMessage);									\
-	}																											\
+#define MRG_LOG_ONCE_PER_RENDER(Identifier, LogCategory, LogVerbosity, LogFormat, ...)									\
+	static UWorld* PREPROCESSOR_JOIN(__LastRenderWorld__, Identifier) = nullptr;										\
+	static TArray<FString> PREPROCESSOR_JOIN(__LoggedMessages__, Identifier);											\
+																														\
+	/* If this world has not been encountered yet (ie, no renders yet, or this is a new render), clear out */			\
+	/* the messages that have been recorded as already logged. */														\
+	if (PREPROCESSOR_JOIN(__LastRenderWorld__, Identifier) != MoviePipeline::FindCurrentWorld())						\
+	{																													\
+		PREPROCESSOR_JOIN(__LastRenderWorld__, Identifier) = MoviePipeline::FindCurrentWorld();							\
+		PREPROCESSOR_JOIN(__LoggedMessages__, Identifier).Empty();														\
+	}																													\
+																														\
+	/* If this message hasn't been logged, record that it has been logged so we can skip it for other frames. */		\
+	const FString PREPROCESSOR_JOIN(__FinalLogMsg__, Identifier) = FString::Printf(LogFormat, ##__VA_ARGS__);			\
+	if (!PREPROCESSOR_JOIN(__LoggedMessages__, Identifier).Contains(PREPROCESSOR_JOIN(__FinalLogMsg__, Identifier)))	\
+	{																													\
+		UE_LOG(LogCategory, LogVerbosity, TEXT("%s"), *PREPROCESSOR_JOIN(__FinalLogMsg__, Identifier));					\
+		PREPROCESSOR_JOIN(__LoggedMessages__, Identifier).Add(PREPROCESSOR_JOIN(__FinalLogMsg__, Identifier));			\
+	}																													\
 
 namespace UE
 {
diff --git a/Engine/Plugins/MovieScene/MovieRenderPipeline/Source/MovieRenderPipelineEditor/Private/Customizations/Graph/MovieGraphMultiLayerEXRNodeCustomization.h b/Engine/Plugins/MovieScene/MovieRenderPipeline/Source/MovieRenderPipelineEditor/Private/Customizations/Graph/MovieGraphMultiLayerEXRNodeCustomization.h
new file mode 100644
--- /dev/null
+++ b/Engine/Plugins/MovieScene/MovieRenderPipeline/Source/MovieRenderPipelineEditor/Private/Customizations/Graph/MovieGraphMultiLayerEXRNodeCustomization.h
@@ -0,0 +1,46 @@
+// Copyright Epic Games, Inc. All Rights Reserved.
+
+#pragma once
+
+#include "DetailLayoutBuilder.h"
+#include "IDetailCustomization.h"
+#include "PropertyHandle.h"
+
+#include "MovieGraphImageSequenceOutputNode.h"
+#include "MovieRenderPipelineCoreModule.h"
+
+/* Customizes how the multi-layer EXR node's properties appear in the details panel. */
+class FMovieGraphMultiLayerEXRNodeCustomization : public IDetailCustomization
+{
+public:
+	static TSharedRef<IDetailCustomization> MakeInstance()
+	{
+		return MakeShared<FMovieGraphMultiLayerEXRNodeCustomization>();
+	}
+
+protected:
+	//~ Begin IDetailCustomization interface
+	virtual void CustomizeDetails(const TSharedPtr<IDetailLayoutBuilder>& DetailBuilder) override
+	{
+		CustomizeDetails(*DetailBuilder);
+	}
+	
+	virtual void CustomizeDetails(IDetailLayoutBuilder& DetailBuilder) override
+	{
+		// These properties are only relevant to single-layer EXRs, not multi-layer.
+		const TSharedRef<IPropertyHandle> CompositeOntoFinalImageProperty = DetailBuilder.GetProperty(GET_MEMBER_NAME_CHECKED(UMovieGraphImageSequenceOutputNode_EXR, bCompositeOntoFinalImage));
+		const TSharedRef<IPropertyHandle> BurnInFileNameFormatProperty = DetailBuilder.GetProperty(GET_MEMBER_NAME_CHECKED(UMovieGraphImageSequenceOutputNode_EXR, BurnInFileNameFormat));
+
+		// The multi-layer EXR node inherits from the normal EXR node; this customization should only be applied to multi-layer.
+		const TArray<TWeakObjectPtr<UMovieGraphImageSequenceOutputNode_MultiLayerEXR>> MultiLayerExrNodes =
+			DetailBuilder.GetObjectsOfTypeBeingCustomized<UMovieGraphImageSequenceOutputNode_MultiLayerEXR>();
+		if (MultiLayerExrNodes.IsEmpty())
+		{
+			return;
+		}
+
+		DetailBuilder.HideProperty(CompositeOntoFinalImageProperty);
+		DetailBuilder.HideProperty(BurnInFileNameFormatProperty);
+	}
+	//~ End IDetailCustomization interface
+};
\ No newline at end of file
diff --git a/Engine/Plugins/MovieScene/MovieRenderPipeline/Source/MovieRenderPipelineEditor/Private/Graph/MovieGraphAssetToolkit.cpp b/Engine/Plugins/MovieScene/MovieRenderPipeline/Source/MovieRenderPipelineEditor/Private/Graph/MovieGraphAssetToolkit.cpp
--- a/Engine/Plugins/MovieScene/MovieRenderPipeline/Source/MovieRenderPipelineEditor/Private/Graph/MovieGraphAssetToolkit.cpp
+++ b/Engine/Plugins/MovieScene/MovieRenderPipeline/Source/MovieRenderPipelineEditor/Private/Graph/MovieGraphAssetToolkit.cpp
@@ -4,6 +4,7 @@
 
 #include "Customizations/Graph/MovieGraphApplyCVarPresetNodeCustomization.h"
 #include "Customizations/Graph/MovieGraphCollectionsCustomization.h"
+#include "Customizations/Graph/MovieGraphMultiLayerEXRNodeCustomization.h"
 #include "Customizations/Graph/MovieGraphFormatTokenCustomization.h"
 #include "Customizations/Graph/MovieGraphMemberCustomization.h"
 #include "Customizations/Graph/MovieGraphMetadataAttributeCustomization.h"
@@ -15,7 +16,6 @@
 #include "Customizations/Graph/MovieGraphSetCVarValueNodeCustomization.h"
 #include "Customizations/Graph/MovieGraphShowFlagsCustomization.h"
 #include "Customizations/Graph/MovieGraphVersioningSettingsCustomization.h"
-#include "Graph/Renderers/MovieGraphShowFlags.h"
 
 #include "Graph/MovieGraphConfig.h"
 #include "Graph/Nodes/MovieGraphApplyCVarPresetNode.h"
@@ -28,7 +28,9 @@
 #include "Graph/Nodes/MovieGraphSelectNode.h"
 #include "Graph/Nodes/MovieGraphSetMetadataAttributesNode.h"
 #include "Graph/Nodes/MovieGraphSetCVarValueNode.h"
+#include "Graph/Renderers/MovieGraphShowFlags.h"
 #include "MovieEdGraphNode.h"
+#include "MovieGraphImageSequenceOutputNode.h"
 #include "MovieGraphSchema.h"
 #include "MoviePipelineCommands.h"
 #include "MovieRenderPipelineSettings.h"
@@ -416,6 +418,10 @@ TSharedRef<SDockTab> FMovieGraphAssetToolkit::SpawnTab_RenderGraphDetails(const
 	SelectedGraphObjectsDetailsWidget->RegisterInstancedCustomPropertyTypeLayout(
 		FMovieGraphVersioningSettings::StaticStruct()->GetFName(),
 		FOnGetPropertyTypeCustomizationInstance::CreateStatic(&FMovieGraphVersioningSettingsCustomization::MakeInstance));
+
+	SelectedGraphObjectsDetailsWidget->RegisterInstancedCustomPropertyLayout(
+		UMovieGraphImageSequenceOutputNode_EXR::StaticClass(),
+		FOnGetDetailCustomizationInstance::CreateStatic(&FMovieGraphMultiLayerEXRNodeCustomization::MakeInstance));
 	
 	SelectedGraphObjectsDetailsWidget->RegisterInstancedCustomPropertyLayout(
 		UMovieGraphCollectionNode::StaticClass(),
diff --git a/Engine/Plugins/MovieScene/MovieRenderPipeline/Source/MovieRenderPipelineMP4Encoder/Private/Graph/MovieGraphMP4EncoderNode.cpp b/Engine/Plugins/MovieScene/MovieRenderPipeline/Source/MovieRenderPipelineMP4Encoder/Private/Graph/MovieGraphMP4EncoderNode.cpp
--- a/Engine/Plugins/MovieScene/MovieRenderPipeline/Source/MovieRenderPipelineMP4Encoder/Private/Graph/MovieGraphMP4EncoderNode.cpp
+++ b/Engine/Plugins/MovieScene/MovieRenderPipeline/Source/MovieRenderPipelineMP4Encoder/Private/Graph/MovieGraphMP4EncoderNode.cpp
@@ -6,6 +6,7 @@
 #include "Graph/MovieGraphConfig.h"
 #include "Graph/MovieGraphOCIOHelper.h"
 #include "Graph/MovieGraphPipeline.h"
+#include "Graph/Nodes/MovieGraphBurnInNode.h"
 #include "Graph/Nodes/MovieGraphGlobalOutputSettingNode.h"
 #include "ImageWriteTask.h"
 #include "MoviePipelineImageQuantization.h"
@@ -262,3 +263,25 @@ void UMovieGraphMP4EncoderNode::UpdateTelemetry(FMoviePipelineShotRenderTelemetr
 {
 	InTelemetry->bUsesMP4 = true;
 }
+
+void UMovieGraphMP4EncoderNode::InjectNodesPostEvaluation(const FName& InBranchName, UMovieGraphEvaluatedConfig* InEvaluatedConfig, TArray<UMovieGraphSettingNode*>& OutInjectedNodes)
+{
+	if (bEnableBurnIn)
+	{
+		UMovieGraphOutputBurnInNode* BurnInNode = NewObject<UMovieGraphOutputBurnInNode>(InEvaluatedConfig);
+		BurnInNode->bOverride_OutputName = true;
+		BurnInNode->bOverride_FileNameFormat = true;
+		BurnInNode->bOverride_OutputRestriction = true;
+		BurnInNode->bOverride_LayerNameFormat = true;
+		BurnInNode->bOverride_BurnInClass = true;
+		BurnInNode->bOverride_bCompositeOntoFinalImage = true;
+		BurnInNode->OutputName = TEXT("MP4");
+		BurnInNode->FileNameFormat = BurnInFileNameFormat;
+		BurnInNode->OutputRestriction = FSoftClassPath(GetClass());
+		BurnInNode->LayerNameFormat = TEXT("");	// Unused
+		BurnInNode->BurnInClass = BurnInClass;
+		BurnInNode->bCompositeOntoFinalImage = bCompositeOntoFinalImage;
+
+		OutInjectedNodes.Add(BurnInNode);
+	}
+}
\ No newline at end of file
diff --git a/Engine/Plugins/MovieScene/MovieRenderPipeline/Source/MovieRenderPipelineMP4Encoder/Public/Graph/MovieGraphMP4EncoderNode.h b/Engine/Plugins/MovieScene/MovieRenderPipeline/Source/MovieRenderPipelineMP4Encoder/Public/Graph/MovieGraphMP4EncoderNode.h
--- a/Engine/Plugins/MovieScene/MovieRenderPipeline/Source/MovieRenderPipelineMP4Encoder/Public/Graph/MovieGraphMP4EncoderNode.h
+++ b/Engine/Plugins/MovieScene/MovieRenderPipeline/Source/MovieRenderPipelineMP4Encoder/Public/Graph/MovieGraphMP4EncoderNode.h
@@ -12,7 +12,7 @@
 
 /** A node which can output H264 mp4 files. */
 UCLASS(MinimalAPI, BlueprintType)
-class UMovieGraphMP4EncoderNode final : public UMovieGraphVideoOutputNode
+class UMovieGraphMP4EncoderNode final : public UMovieGraphVideoOutputNode, public IMovieGraphEvaluationNodeInjector
 {
 	GENERATED_BODY()
 
@@ -42,6 +42,10 @@ protected:
 	UE_API virtual void UpdateTelemetry(FMoviePipelineShotRenderTelemetry* InTelemetry) const override;
 	// ~UMovieGraphSettingNode Interface
 
+	// IMovieGraphEvaluationNodeInjector Interface
+	virtual void InjectNodesPostEvaluation(const FName& InBranchName, UMovieGraphEvaluatedConfig* InEvaluatedConfig, TArray<UMovieGraphSettingNode*>& OutInjectedNodes) override;
+	// ~IMovieGraphEvaluationNodeInjector Interface
+
 protected:
 	struct FMP4CodecWriter : public MovieRenderGraph::IVideoCodecWriter
 	{
@@ -79,6 +83,18 @@ public:
 
 	UPROPERTY(EditAnywhere, BlueprintReadWrite, Category = Overrides, meta = (InlineEditConditionToggle))
 	uint8 bOverride_OCIOContext : 1;
+
+	UPROPERTY(EditAnywhere, BlueprintReadWrite, Category = Overrides, meta = (InlineEditConditionToggle))
+	uint8 bOverride_bEnableBurnIn : 1;
+
+	UPROPERTY(EditAnywhere, BlueprintReadWrite, Category = Overrides, meta = (InlineEditConditionToggle))
+	uint8 bOverride_BurnInClass : 1;
+
+	UPROPERTY(EditAnywhere, BlueprintReadWrite, Category = Overrides, meta = (InlineEditConditionToggle))
+	uint8 bOverride_bCompositeOntoFinalImage : 1;
+
+	UPROPERTY(EditAnywhere, BlueprintReadWrite, Category = Overrides, meta = (InlineEditConditionToggle))
+	uint8 bOverride_BurnInFileNameFormat : 1;
 		
 	/**
 	* Specifies the bitrate control method used by the encoder. Quality lets the user target a given quality without concern
@@ -161,6 +177,22 @@ public:
 	*/
 	UPROPERTY(EditAnywhere, BlueprintReadWrite, Category = "OCIO", DisplayName = "OCIO Context", meta = (DisplayAfter = "OCIOConfiguration", EditCondition = "bOverride_OCIOContext"))
 	TMap<FString, FString> OCIOContext;
+
+	/** If true, this output node will also generate a burn-in. */
+	UPROPERTY(EditAnywhere, BlueprintReadWrite, Category = "Burn In", meta = (EditCondition = "bOverride_bEnableBurnIn"))
+	bool bEnableBurnIn = false;
+
+	/** The widget that the burn-in should use. */
+	UPROPERTY(EditAnywhere, BlueprintReadWrite, Category="Burn In", meta=(MetaClass="/Script/MovieRenderPipelineCore.MovieGraphBurnInWidget", EditCondition="bOverride_BurnInClass"))
+	FSoftClassPath BurnInClass = FSoftClassPath(TEXT("/MovieRenderPipeline/Blueprints/Graph/DefaultGraphBurnIn.DefaultGraphBurnIn_C"));
+	
+	/** If true, the burn-in that's generated will be composited onto this output. Otherwise, the burn-in will be written to a different file. */
+	UPROPERTY(EditAnywhere, BlueprintReadWrite, Category = "Burn In", meta = (EditCondition = "bOverride_bCompositeOntoFinalImage"))
+	bool bCompositeOntoFinalImage = true;
+
+	/** The file name format used for writing out the burn-in video (if not composited). */
+	UPROPERTY(EditAnywhere, BlueprintReadWrite, Category = "Burn In", meta = (EditCondition = "bOverride_BurnInFileNameFormat"))
+	FString BurnInFileNameFormat = TEXT("{sequence_name}.{layer_name}.{renderer_name}");
 };
 
 #undef UE_API
diff --git a/Engine/Plugins/MovieScene/MovieRenderPipeline/Source/MovieRenderPipelineRenderPasses/Private/MovieGraphImageSequenceOutputNode.cpp b/Engine/Plugins/MovieScene/MovieRenderPipeline/Source/MovieRenderPipelineRenderPasses/Private/MovieGraphImageSequenceOutputNode.cpp
--- a/Engine/Plugins/MovieScene/MovieRenderPipeline/Source/MovieRenderPipelineRenderPasses/Private/MovieGraphImageSequenceOutputNode.cpp
+++ b/Engine/Plugins/MovieScene/MovieRenderPipeline/Source/MovieRenderPipelineRenderPasses/Private/MovieGraphImageSequenceOutputNode.cpp
@@ -2,6 +2,7 @@
 
 #include "MovieGraphImageSequenceOutputNode.h"
 
+#include "Graph/Nodes/MovieGraphBurnInNode.h"
 #include "Graph/Nodes/MovieGraphGlobalOutputSettingNode.h"
 #include "Graph/Nodes/MovieGraphRenderLayerNode.h"
 #include "Graph/MovieGraphDataTypes.h"
@@ -23,6 +24,44 @@
 
 #include UE_INLINE_GENERATED_CPP_BY_NAME(MovieGraphImageSequenceOutputNode)
 
+#define INJECT_BURN_IN_NODE(OUTPUT_NAME)																		\
+	if (bEnableBurnIn)																							\
+	{																											\
+		UMovieGraphOutputBurnInNode* BurnInNode = NewObject<UMovieGraphOutputBurnInNode>(InEvaluatedConfig);	\
+		BurnInNode->bOverride_OutputName = true;																\
+		BurnInNode->bOverride_FileNameFormat = true;															\
+		BurnInNode->bOverride_OutputRestriction = true;															\
+		BurnInNode->bOverride_bCompositeOntoFinalImage = true;													\
+		BurnInNode->bOverride_BurnInClass = true;																\
+		BurnInNode->OutputName = OUTPUT_NAME;																	\
+		BurnInNode->FileNameFormat = BurnInFileNameFormat;														\
+		BurnInNode->OutputRestriction = FSoftClassPath(GetClass());												\
+		BurnInNode->bCompositeOntoFinalImage = bCompositeOntoFinalImage;										\
+		BurnInNode->BurnInClass = BurnInClass;																	\
+																												\
+		OutInjectedNodes.Add(BurnInNode);																		\
+	}																											\
+
+#define INJECT_BURN_IN_MULTILAYER_NODE(OUTPUT_NAME)																\
+	if (bEnableBurnIn)																							\
+	{																											\
+		UMovieGraphOutputBurnInNode* BurnInNode = NewObject<UMovieGraphOutputBurnInNode>(InEvaluatedConfig);	\
+		BurnInNode->bOverride_OutputName = true;																\
+		BurnInNode->bOverride_FileNameFormat = true;															\
+		BurnInNode->bOverride_OutputRestriction = true;															\
+		BurnInNode->bOverride_bCompositeOntoFinalImage = true;													\
+		BurnInNode->bOverride_BurnInClass = true;																\
+		BurnInNode->bOverride_LayerNameFormat = true;															\
+		BurnInNode->OutputName = OUTPUT_NAME;																	\
+		BurnInNode->FileNameFormat = BurnInFileNameFormat;														\
+		BurnInNode->OutputRestriction = FSoftClassPath(GetClass());												\
+		BurnInNode->bCompositeOntoFinalImage = false;															\
+		BurnInNode->BurnInClass = BurnInClass;																	\
+		BurnInNode->LayerNameFormat = BurnInLayerNameFormat;													\
+																												\
+		OutInjectedNodes.Add(BurnInNode);																		\
+	}																											\
+
 UMovieGraphImageSequenceOutputNode::UMovieGraphImageSequenceOutputNode()
 {
 	ImageWriteQueue = &FModuleManager::Get().LoadModuleChecked<IImageWriteQueueModule>("ImageWriteQueue").GetWriteQueue();
@@ -64,8 +103,13 @@ FString UMovieGraphImageSequenceOutputNode::CreateFileName(
 		return FString();
 	}
 
+	const UE::MovieGraph::FMovieGraphSampleState* Payload = InRenderData.Value->GetPayload<UE::MovieGraph::FMovieGraphSampleState>();
+
+	// The file name format usually comes from the output node directly, but the payload has a chance to override it.
+	const FString FinalFileNameFormat = !Payload->FilenameFormatOverride.IsEmpty() ? Payload->FilenameFormatOverride : InParentNode->FileNameFormat;
+
 	// Generate one string that puts the directory combined with the filename format.
-	FString FileNameFormatString = OutputSettingNode->OutputDirectory.Path / InParentNode->FileNameFormat;
+	FString FileNameFormatString = OutputSettingNode->OutputDirectory.Path / FinalFileNameFormat;
 
 	// Insert tokens like {layer_name} as appropriate to make sure outputs don't clash with each other.
 	DisambiguateFilename(FileNameFormatString, InRawFrameData, InParentNode->GetFName(), InRenderData);
@@ -80,8 +124,6 @@ FString UMovieGraphImageSequenceOutputNode::CreateFileName(
 	TMap<FString, FString> AdditionalFormatArgs;
 	AdditionalFormatArgs.Add(TEXT("ext"), Extension);
 
-	UE::MovieGraph::FMovieGraphSampleState* Payload = InRenderData.Value->GetPayload<UE::MovieGraph::FMovieGraphSampleState>();
-
 	FMovieGraphFilenameResolveParams Params = FMovieGraphFilenameResolveParams::MakeResolveParams(
 		InRenderData.Key,
 		InPipeline,
@@ -150,6 +192,14 @@ void UMovieGraphImageSequenceOutputNode::OnReceiveImageDataImpl(UMovieGraphPipel
 	// The base ImageSequenceOutputNode doesn't support any multilayer formats, so we write out each render pass separately.
 	for (TPair<FMovieGraphRenderDataIdentifier, TUniquePtr<FImagePixelData>>& RenderData : InRawFrameData->ImageOutputData)
 	{
+		UE::MovieGraph::FMovieGraphSampleState* Payload = RenderData.Value->GetPayload<UE::MovieGraph::FMovieGraphSampleState>();
+
+		// This payload may have opted out of being written by this output type.
+		if (!Payload->OutputTypeRestrictions.IsEmpty() && !Payload->OutputTypeRestrictions.Contains(FSoftClassPath(GetClass())))
+		{
+			continue;
+		}
+		
 		// If this pass is composited, skip it for now
 		if (CompositedPasses.ContainsByPredicate([&RenderData](const TPair<FMovieGraphRenderDataIdentifier, TUniquePtr<FImagePixelData>>& CompositedPass)
 			{
@@ -200,8 +250,6 @@ void UMovieGraphImageSequenceOutputNode::OnReceiveImageDataImpl(UMovieGraphPipel
 			TileImageTask->PixelData = RenderData.Value->MoveImageDataToNew();
 		}
 
-		UE::MovieGraph::FMovieGraphSampleState* Payload = RenderData.Value->GetPayload<UE::MovieGraph::FMovieGraphSampleState>();
-
 		// If the overscan isn't cropped from the final image, offset any composites to the top-left of the original frustum
 		FIntPoint CompositeOffset = FIntPoint::ZeroValue;
 		const bool bIsCropRectValid = !Payload->CropRectangle.IsEmpty();
@@ -259,6 +307,14 @@ void UMovieGraphImageSequenceOutputNode::OnReceiveImageDataImpl(UMovieGraphPipel
 			{
 				continue;
 			}
+			
+			UE::MovieGraph::FMovieGraphSampleState* CompositePayload = CompositedPass.Value->GetPayload<UE::MovieGraph::FMovieGraphSampleState>();
+
+			// This composite may have opted out of being written by this output type.
+			if (!CompositePayload->OutputTypeRestrictions.IsEmpty() && !CompositePayload->OutputTypeRestrictions.Contains(FSoftClassPath(GetClass())))
+			{
+				continue;
+			}
 
 			// There could be multiple renders within this branch using the composited pass, so we have to copy the image data
 			switch (PixelType)
@@ -394,6 +450,14 @@ void UMovieGraphImageSequenceOutputNode_EXR::OnReceiveImageDataImpl(UMovieGraphP
 
 	for (TPair<FMovieGraphRenderDataIdentifier, TUniquePtr<FImagePixelData>>& RenderData : InRawFrameData->ImageOutputData)
 	{
+		UE::MovieGraph::FMovieGraphSampleState* Payload = RenderData.Value->GetPayload<UE::MovieGraph::FMovieGraphSampleState>();
+
+		// This payload may have opted out of being written by this output type.
+		if (!Payload->OutputTypeRestrictions.IsEmpty() && !Payload->OutputTypeRestrictions.Contains(FSoftClassPath(GetClass())))
+		{
+			continue;
+		}
+		
 		// If this pass is composited, skip it for now
 		if (CompositedPasses.ContainsByPredicate([&RenderData](const TPair<FMovieGraphRenderDataIdentifier, TUniquePtr<FImagePixelData>>& CompositedPass)
 			{
@@ -425,8 +489,6 @@ void UMovieGraphImageSequenceOutputNode_EXR::OnReceiveImageDataImpl(UMovieGraphP
 			continue;
 		}
 
-		UE::MovieGraph::FMovieGraphSampleState* Payload = RenderData.Value->GetPayload<UE::MovieGraph::FMovieGraphSampleState>();
-
 		// The pass may be need to be written out with lossless compression (eg, Object ID forces this on).
 		const EEXRCompressionFormat CompressionType = Payload->bForceLosslessCompression ? ParentNode->LosslessCompression : ParentNode->Compression;
 		
@@ -467,6 +529,14 @@ void UMovieGraphImageSequenceOutputNode_EXR::OnReceiveImageDataImpl(UMovieGraphP
 			{
 				continue;
 			}
+
+			UE::MovieGraph::FMovieGraphSampleState* CompositePayload = CompositedPass.Value->GetPayload<UE::MovieGraph::FMovieGraphSampleState>();
+
+			// This composite may have opted out of being written by this output type.
+			if (!CompositePayload->OutputTypeRestrictions.IsEmpty() && !CompositePayload->OutputTypeRestrictions.Contains(FSoftClassPath(GetClass())))
+			{
+				continue;
+			}
 			
 			// This composited pass will only composite on top of renders w/ the same branch and camera
 			if (CompositedPass.Key.IsBranchAndCameraEqual(RenderData.Key))
@@ -503,6 +573,10 @@ void UMovieGraphImageSequenceOutputNode_EXR::OnReceiveImageDataImpl(UMovieGraphP
 	}
 }
 
+void UMovieGraphImageSequenceOutputNode_EXR::InjectNodesPostEvaluation(const FName& InBranchName, UMovieGraphEvaluatedConfig* InEvaluatedConfig, TArray<UMovieGraphSettingNode*>& OutInjectedNodes)
+{
+	INJECT_BURN_IN_NODE(TEXT("EXR"))
+}
 
 void UMovieGraphImageSequenceOutputNode_MultiLayerEXR::OnReceiveImageDataImpl(UMovieGraphPipeline* InPipeline, UE::MovieGraph::FMovieGraphOutputMergerFrame* InRawFrameData, const TSet<FMovieGraphRenderDataIdentifier>& InMask)
 {
@@ -575,6 +649,13 @@ void UMovieGraphImageSequenceOutputNode_MultiLayerEXR::OnReceiveImageDataImpl(UM
 			checkf(ImageData.IsValid(), TEXT("Unexpected empty image data: incorrectly moved or its production failed?"));
 
 			const UE::MovieGraph::FMovieGraphSampleState* Payload = ImageData->GetPayload<UE::MovieGraph::FMovieGraphSampleState>();
+
+			// This payload may have opted out of being written by this output type.
+			if (!Payload->OutputTypeRestrictions.IsEmpty() && !Payload->OutputTypeRestrictions.Contains(FSoftClassPath(GetClass())))
+			{
+				continue;
+			}
+			
 			ShotIndex = Payload->TraversalContext.ShotIndex;
 
 			LowestRenderLayerIndex = FMath::Min(LowestRenderLayerIndex, Payload->RenderLayerIndex);
@@ -688,6 +769,11 @@ void UMovieGraphImageSequenceOutputNode_MultiLayerEXR::OnReceiveImageDataImpl(UM
 
 #endif // WITH_UNREALEXR
 
+void UMovieGraphImageSequenceOutputNode_MultiLayerEXR::InjectNodesPostEvaluation(const FName& InBranchName, UMovieGraphEvaluatedConfig* InEvaluatedConfig, TArray<UMovieGraphSettingNode*>& OutInjectedNodes)
+{
+	INJECT_BURN_IN_MULTILAYER_NODE(TEXT("MULTI_EXR"))
+}
+
 void UMovieGraphImageSequenceOutputNode_MultiLayerEXR::GetFilenameToEXROutputConfigMappings(
 	const UMovieGraphImageSequenceOutputNode_MultiLayerEXR* InParentNode, UMovieGraphPipeline* InPipeline,
 	UE::MovieGraph::FMovieGraphOutputMergerFrame* InRawFrameData, TMap<FString, FEXROutputConfigForFilename>& OutFilenameToOutputConfigs) const
@@ -801,3 +887,18 @@ FString UMovieGraphImageSequenceOutputNode_MultiLayerEXR::ResolveOutputFilename(
 
 	return FinalFilePath;
 }
+
+void UMovieGraphImageSequenceOutputNode_BMP::InjectNodesPostEvaluation(const FName& InBranchName, UMovieGraphEvaluatedConfig* InEvaluatedConfig, TArray<UMovieGraphSettingNode*>& OutInjectedNodes)
+{
+	INJECT_BURN_IN_NODE(TEXT("BMP"))
+}
+
+void UMovieGraphImageSequenceOutputNode_JPG::InjectNodesPostEvaluation(const FName& InBranchName, UMovieGraphEvaluatedConfig* InEvaluatedConfig, TArray<UMovieGraphSettingNode*>& OutInjectedNodes)
+{
+	INJECT_BURN_IN_NODE(TEXT("JPG"))
+}
+
+void UMovieGraphImageSequenceOutputNode_PNG::InjectNodesPostEvaluation(const FName& InBranchName, UMovieGraphEvaluatedConfig* InEvaluatedConfig, TArray<UMovieGraphSettingNode*>& OutInjectedNodes)
+{
+	INJECT_BURN_IN_NODE(TEXT("PNG"))
+}
diff --git a/Engine/Plugins/MovieScene/MovieRenderPipeline/Source/MovieRenderPipelineRenderPasses/Public/MovieGraphImageSequenceOutputNode.h b/Engine/Plugins/MovieScene/MovieRenderPipeline/Source/MovieRenderPipelineRenderPasses/Public/MovieGraphImageSequenceOutputNode.h
--- a/Engine/Plugins/MovieScene/MovieRenderPipeline/Source/MovieRenderPipelineRenderPasses/Public/MovieGraphImageSequenceOutputNode.h
+++ b/Engine/Plugins/MovieScene/MovieRenderPipeline/Source/MovieRenderPipelineRenderPasses/Public/MovieGraphImageSequenceOutputNode.h
@@ -2,6 +2,7 @@
 
 #pragma once
 
+#include "Graph/MovieGraphConfig.h"
 #include "Graph/Nodes/MovieGraphFileOutputNode.h"
 
 #include "Async/Future.h"
@@ -94,8 +95,8 @@ protected:
 /**
  * Image sequence output node that can write single-layer EXR files.
  */
-UCLASS(MinimalAPI)
-class UMovieGraphImageSequenceOutputNode_EXR : public UMovieGraphImageSequenceOutputNode
+UCLASS(MinimalAPI, meta=(PrioritizeCategories="EXR OCIO"))
+class UMovieGraphImageSequenceOutputNode_EXR : public UMovieGraphImageSequenceOutputNode, public IMovieGraphEvaluationNodeInjector
 {
 	GENERATED_BODY()
 
@@ -141,6 +142,10 @@ public:
 	}
 #endif
 	
+	//~ Begin IMovieGraphEvaluationNodeInjector interface
+	virtual void InjectNodesPostEvaluation(const FName& InBranchName, UMovieGraphEvaluatedConfig* InEvaluatedConfig, TArray<UMovieGraphSettingNode*>& OutInjectedNodes) override;
+	//~ End IMovieGraphEvaluationNodeInjector interface
+	
 public:
 	UPROPERTY(EditAnywhere, BlueprintReadWrite, Category = Overrides, meta = (InlineEditConditionToggle))
 	uint8 bOverride_Compression : 1;
@@ -148,6 +153,18 @@ public:
 	UPROPERTY(EditAnywhere, BlueprintReadWrite, Category = Overrides, meta = (InlineEditConditionToggle))
 	uint8 bOverride_LosslessCompression : 1;
 	
+	UPROPERTY(EditAnywhere, BlueprintReadWrite, Category = Overrides, meta = (InlineEditConditionToggle))
+	uint8 bOverride_bEnableBurnIn : 1;
+
+	UPROPERTY(EditAnywhere, BlueprintReadWrite, Category = Overrides, meta = (InlineEditConditionToggle))
+	uint8 bOverride_BurnInClass : 1;
+
+	UPROPERTY(EditAnywhere, BlueprintReadWrite, Category = Overrides, meta = (InlineEditConditionToggle))
+	uint8 bOverride_bCompositeOntoFinalImage : 1;
+
+	UPROPERTY(EditAnywhere, BlueprintReadWrite, Category = Overrides, meta = (InlineEditConditionToggle))
+	uint8 bOverride_BurnInFileNameFormat : 1;
+
 	/**
 	 * The type of compression the EXR file should be compressed with. This compression method will be applied to all layers unless the data in
 	 * a layer has been flagged as needing lossless compression. See "Lossless Compression" for more details.
@@ -162,6 +179,22 @@ public:
 	UPROPERTY(EditAnywhere, BlueprintReadWrite, meta=(EditCondition="bOverride_LosslessCompression", ValidEnumValues="None,RLE,ZIPS,ZIP,PIZ"), Category = "EXR")
 	EEXRCompressionFormat LosslessCompression;
 
+	/** If true, this output node will also generate a burn-in. */
+	UPROPERTY(EditAnywhere, BlueprintReadWrite, Category = "Burn In", meta = (EditCondition = "bOverride_bEnableBurnIn"))
+	bool bEnableBurnIn = false;
+
+	/** The widget that the burn-in should use. */
+	UPROPERTY(EditAnywhere, BlueprintReadWrite, Category="Burn In", meta=(MetaClass="/Script/MovieRenderPipelineCore.MovieGraphBurnInWidget", EditCondition="bOverride_BurnInClass"))
+	FSoftClassPath BurnInClass = FSoftClassPath(TEXT("/MovieRenderPipeline/Blueprints/Graph/DefaultGraphBurnIn.DefaultGraphBurnIn_C"));
+
+	/** If true, the burn-in that's generated will be composited onto this output. Otherwise, the burn-in will be written to a different file. */
+	UPROPERTY(EditAnywhere, BlueprintReadWrite, Category = "Burn In", meta = (EditCondition = "bOverride_bCompositeOntoFinalImage"))
+	bool bCompositeOntoFinalImage = true;
+
+	/** The file name format used for writing out the burn-in image (if not composited). */
+	UPROPERTY(EditAnywhere, BlueprintReadWrite, Category = "Burn In", meta = (EditCondition = "bOverride_BurnInFileNameFormat"))
+	FString BurnInFileNameFormat = TEXT("{sequence_name}.{layer_name}.{renderer_name}.{frame_number}");
+
 protected:
 
 #if WITH_UNREALEXR
@@ -197,7 +230,7 @@ protected:
 /**
  * Image sequence output node that can write multi-layer EXR files.
  */
-UCLASS(MinimalAPI)
+UCLASS(MinimalAPI, meta=(PrioritizeCategories="EXR OCIO"))
 class UMovieGraphImageSequenceOutputNode_MultiLayerEXR : public UMovieGraphImageSequenceOutputNode_EXR
 {
 	GENERATED_BODY()
@@ -252,6 +285,10 @@ public:
 		return NSLOCTEXT("MovieGraphNodes", "FileOutputGraphNode_Category", "Output Type");
 	}
 #endif
+	
+	//~ Begin IMovieGraphEvaluationNodeInjector interface
+	virtual void InjectNodesPostEvaluation(const FName& InBranchName, UMovieGraphEvaluatedConfig* InEvaluatedConfig, TArray<UMovieGraphSettingNode*>& OutInjectedNodes) override;
+	//~ End IMovieGraphEvaluationNodeInjector interface
 
 private:
 	/** Stores the EXR output config to use for a specific file name */
@@ -287,6 +324,9 @@ private:
 public:
 	UPROPERTY(EditAnywhere, BlueprintReadWrite, Category = Overrides, meta = (InlineEditConditionToggle))
 	uint8 bOverride_bMultiPart : 1;
+
+	UPROPERTY(EditAnywhere, BlueprintReadWrite, Category = Overrides, meta = (InlineEditConditionToggle))
+	uint8 bOverride_BurnInLayerNameFormat : 1;
 	
 	/**
 	 * Indicates whether the exr file should be written as a multi-part exr file, which supports having different image types and resolutions for each layer.
@@ -298,6 +338,10 @@ public:
 	 */
 	UPROPERTY(EditAnywhere, BlueprintReadWrite, Category = "EXR", meta=(EditCondition="bOverride_bMultiPart"))
 	bool bMultipart = true;
+
+	/** The layer name used for this burn-in inside of multi-layer EXRs. */
+	UPROPERTY(EditAnywhere, BlueprintReadWrite, Category = "Burn In", meta = (DisplayAfter = "BurnInClass", EditCondition = "bOverride_BurnInLayerNameFormat"))
+	FString BurnInLayerNameFormat;
 };
 
 /**
@@ -305,8 +349,8 @@ public:
 * be useful in rare occasions (bmp files are uncompressed but larger). sRGB is applied.
 * No metadata is supported.
 */
-UCLASS(MinimalAPI)
-class UMovieGraphImageSequenceOutputNode_BMP : public UMovieGraphImageSequenceOutputNode
+UCLASS(MinimalAPI, meta=(PrioritizeCategories="OCIO"))
+class UMovieGraphImageSequenceOutputNode_BMP : public UMovieGraphImageSequenceOutputNode, public IMovieGraphEvaluationNodeInjector
 {
 	GENERATED_BODY()
 
@@ -346,6 +390,39 @@ public:
 		return ImageSequenceIcon;
 	}
 #endif
+	
+	//~ Begin IMovieGraphEvaluationNodeInjector interface
+	virtual void InjectNodesPostEvaluation(const FName& InBranchName, UMovieGraphEvaluatedConfig* InEvaluatedConfig, TArray<UMovieGraphSettingNode*>& OutInjectedNodes) override;
+	//~ End IMovieGraphEvaluationNodeInjector interface
+
+public:
+	UPROPERTY(EditAnywhere, BlueprintReadWrite, Category = Overrides, meta = (InlineEditConditionToggle))
+	uint8 bOverride_bEnableBurnIn : 1;
+
+	UPROPERTY(EditAnywhere, BlueprintReadWrite, Category = Overrides, meta = (InlineEditConditionToggle))
+	uint8 bOverride_BurnInClass : 1;
+
+	UPROPERTY(EditAnywhere, BlueprintReadWrite, Category = Overrides, meta = (InlineEditConditionToggle))
+	uint8 bOverride_bCompositeOntoFinalImage : 1;
+
+	UPROPERTY(EditAnywhere, BlueprintReadWrite, Category = Overrides, meta = (InlineEditConditionToggle))
+	uint8 bOverride_BurnInFileNameFormat : 1;
+
+	/** If true, this output node will also generate a burn-in. */
+    UPROPERTY(EditAnywhere, BlueprintReadWrite, Category = "Burn In", meta = (EditCondition = "bOverride_bEnableBurnIn"))
+    bool bEnableBurnIn = false;
+
+    /** The widget that the burn-in should use. */
+    UPROPERTY(EditAnywhere, BlueprintReadWrite, Category="Burn In", meta=(MetaClass="/Script/MovieRenderPipelineCore.MovieGraphBurnInWidget", EditCondition="bOverride_BurnInClass"))
+    FSoftClassPath BurnInClass = FSoftClassPath(TEXT("/MovieRenderPipeline/Blueprints/Graph/DefaultGraphBurnIn.DefaultGraphBurnIn_C"));
+
+	/** If true, the burn-in that's generated will be composited onto this output. Otherwise, the burn-in will be written to a different file. */
+	UPROPERTY(EditAnywhere, BlueprintReadWrite, Category = "Burn In", meta = (EditCondition = "bOverride_bCompositeOntoFinalImage"))
+	bool bCompositeOntoFinalImage = true;
+
+	/** The file name format used for writing out the burn-in image (if not composited). */
+	UPROPERTY(EditAnywhere, BlueprintReadWrite, Category = "Burn In", meta = (EditCondition = "bOverride_BurnInFileNameFormat"))
+	FString BurnInFileNameFormat = TEXT("{sequence_name}.{layer_name}.{renderer_name}.{frame_number}");
 };
 
 
@@ -354,8 +431,8 @@ public:
 * are lossy, but a good balance between compression speed and final filesize. sRGB is applied.
 * No metadata is supported.
 */
-UCLASS(MinimalAPI)
-class UMovieGraphImageSequenceOutputNode_JPG : public UMovieGraphImageSequenceOutputNode
+UCLASS(MinimalAPI, meta=(PrioritizeCategories="OCIO"))
+class UMovieGraphImageSequenceOutputNode_JPG : public UMovieGraphImageSequenceOutputNode, public IMovieGraphEvaluationNodeInjector
 {
 	GENERATED_BODY()
 
@@ -395,6 +472,39 @@ public:
 		return ImageSequenceIcon;
 	}
 #endif
+
+	//~ Begin IMovieGraphEvaluationNodeInjector interface
+	virtual void InjectNodesPostEvaluation(const FName& InBranchName, UMovieGraphEvaluatedConfig* InEvaluatedConfig, TArray<UMovieGraphSettingNode*>& OutInjectedNodes) override;
+	//~ End IMovieGraphEvaluationNodeInjector interface
+
+public:
+	UPROPERTY(EditAnywhere, BlueprintReadWrite, Category = Overrides, meta = (InlineEditConditionToggle))
+	uint8 bOverride_bEnableBurnIn : 1;
+
+	UPROPERTY(EditAnywhere, BlueprintReadWrite, Category = Overrides, meta = (InlineEditConditionToggle))
+	uint8 bOverride_BurnInClass : 1;
+
+	UPROPERTY(EditAnywhere, BlueprintReadWrite, Category = Overrides, meta = (InlineEditConditionToggle))
+	uint8 bOverride_bCompositeOntoFinalImage : 1;
+
+	UPROPERTY(EditAnywhere, BlueprintReadWrite, Category = Overrides, meta = (InlineEditConditionToggle))
+	uint8 bOverride_BurnInFileNameFormat : 1;
+
+	/** If true, this output node will also generate a burn-in. */
+	UPROPERTY(EditAnywhere, BlueprintReadWrite, Category = "Burn In", meta = (EditCondition = "bOverride_bEnableBurnIn"))
+	bool bEnableBurnIn = false;
+
+	/** The widget that the burn-in should use. */
+	UPROPERTY(EditAnywhere, BlueprintReadWrite, Category="Burn In", meta=(MetaClass="/Script/MovieRenderPipelineCore.MovieGraphBurnInWidget", EditCondition="bOverride_BurnInClass"))
+	FSoftClassPath BurnInClass = FSoftClassPath(TEXT("/MovieRenderPipeline/Blueprints/Graph/DefaultGraphBurnIn.DefaultGraphBurnIn_C"));
+
+	/** If true, the burn-in that's generated will be composited onto this output. Otherwise, the burn-in will be written to a different file. */
+	UPROPERTY(EditAnywhere, BlueprintReadWrite, Category = "Burn In", meta = (EditCondition = "bOverride_bCompositeOntoFinalImage"))
+	bool bCompositeOntoFinalImage = true;
+
+	/** The file name format used for writing out the burn-in image (if not composited). */
+	UPROPERTY(EditAnywhere, BlueprintReadWrite, Category = "Burn In", meta = (EditCondition = "bOverride_BurnInFileNameFormat"))
+	FString BurnInFileNameFormat = TEXT("{sequence_name}.{layer_name}.{renderer_name}.{frame_number}");
 };
 
 /**
@@ -402,8 +512,8 @@ public:
 * are lossless but slow to compress and have a larger final filesize than JPEG. sRGB is applied.
 * No metadata is supported.
 */
-UCLASS(MinimalAPI)
-class UMovieGraphImageSequenceOutputNode_PNG : public UMovieGraphImageSequenceOutputNode
+UCLASS(MinimalAPI, meta=(PrioritizeCategories="OCIO"))
+class UMovieGraphImageSequenceOutputNode_PNG : public UMovieGraphImageSequenceOutputNode, public IMovieGraphEvaluationNodeInjector
 {
 	GENERATED_BODY()
 
@@ -445,6 +555,39 @@ public:
 		return ImageSequenceIcon;
 	}
 #endif
+	
+	//~ Begin IMovieGraphEvaluationNodeInjector interface
+	virtual void InjectNodesPostEvaluation(const FName& InBranchName, UMovieGraphEvaluatedConfig* InEvaluatedConfig, TArray<UMovieGraphSettingNode*>& OutInjectedNodes) override;
+	//~ End IMovieGraphEvaluationNodeInjector interface
+
+public:
+	UPROPERTY(EditAnywhere, BlueprintReadWrite, Category = Overrides, meta = (InlineEditConditionToggle))
+	uint8 bOverride_bEnableBurnIn : 1;
+
+	UPROPERTY(EditAnywhere, BlueprintReadWrite, Category = Overrides, meta = (InlineEditConditionToggle))
+	uint8 bOverride_BurnInClass : 1;
+
+	UPROPERTY(EditAnywhere, BlueprintReadWrite, Category = Overrides, meta = (InlineEditConditionToggle))
+	uint8 bOverride_bCompositeOntoFinalImage : 1;
+
+	UPROPERTY(EditAnywhere, BlueprintReadWrite, Category = Overrides, meta = (InlineEditConditionToggle))
+	uint8 bOverride_BurnInFileNameFormat : 1;
+
+	/** If true, this output node will also generate a burn-in. */
+	UPROPERTY(EditAnywhere, BlueprintReadWrite, Category = "Burn In", meta = (EditCondition = "bOverride_bEnableBurnIn"))
+	bool bEnableBurnIn = false;
+
+	/** The widget that the burn-in should use. */
+	UPROPERTY(EditAnywhere, BlueprintReadWrite, Category="Burn In", meta=(MetaClass="/Script/MovieRenderPipelineCore.MovieGraphBurnInWidget", EditCondition="bOverride_BurnInClass"))
+	FSoftClassPath BurnInClass = FSoftClassPath(TEXT("/MovieRenderPipeline/Blueprints/Graph/DefaultGraphBurnIn.DefaultGraphBurnIn_C"));
+
+	/** If true, the burn-in that's generated will be composited onto this output. Otherwise, the burn-in will be written to a different file. */
+	UPROPERTY(EditAnywhere, BlueprintReadWrite, Category = "Burn In", meta = (EditCondition = "bOverride_bCompositeOntoFinalImage"))
+	bool bCompositeOntoFinalImage = true;
+
+	/** The file name format used for writing out the burn-in image (if not composited). */
+	UPROPERTY(EditAnywhere, BlueprintReadWrite, Category = "Burn In", meta = (EditCondition = "bOverride_BurnInFileNameFormat"))
+	FString BurnInFileNameFormat = TEXT("{sequence_name}.{layer_name}.{renderer_name}.{frame_number}");
 };
 
 #undef UE_API
