MoviePipeline: Fixed CIS error.

diff --git a/Engine/Plugins/MovieScene/MovieRenderPipeline/Source/MovieRenderPipelineRenderPasses/Private/MovieGraphImageSequenceOutputNode.cpp b/Engine/Plugins/MovieScene/MovieRenderPipeline/Source/MovieRenderPipelineRenderPasses/Private/MovieGraphImageSequenceOutputNode.cpp
--- a/Engine/Plugins/MovieScene/MovieRenderPipeline/Source/MovieRenderPipelineRenderPasses/Private/MovieGraphImageSequenceOutputNode.cpp
+++ b/Engine/Plugins/MovieScene/MovieRenderPipeline/Source/MovieRenderPipelineRenderPasses/Private/MovieGraphImageSequenceOutputNode.cpp
@@ -346,7 +346,6 @@ void UMovieGraphImageSequenceOutputNode::OnReceiveImageDataImpl(UMovieGraphPipel
 }
 
 #if WITH_UNREALEXR
-
 TUniquePtr<FEXRImageWriteTask> UMovieGraphImageSequenceOutputNode_EXR::CreateImageWriteTask(FString InFileName, EEXRCompressionFormat InCompression, bool bMultiPart) const
 {
 	// Ensure our OpenExrRTTI module gets loaded.
@@ -441,7 +440,9 @@ void UMovieGraphImageSequenceOutputNode_EXR::UpdateTaskPerLayer(
 
 	InOutImageTask.Layers.Add(MoveTemp(InImageData));
 }
+#endif // WITH_UNREALEXR
 
+#if WITH_UNREALEXR
 void UMovieGraphImageSequenceOutputNode_EXR::OnReceiveImageDataImpl(UMovieGraphPipeline* InPipeline, UE::MovieGraph::FMovieGraphOutputMergerFrame* InRawFrameData, const TSet<FMovieGraphRenderDataIdentifier>& InMask)
 {
 	check(InRawFrameData);
@@ -572,12 +573,14 @@ void UMovieGraphImageSequenceOutputNode_EXR::OnReceiveImageDataImpl(UMovieGraphP
 		InPipeline->AddOutputFuture(MoveTemp(Future), OutputFutureData);
 	}
 }
+#endif // WITH_UNREALEXR
 
 void UMovieGraphImageSequenceOutputNode_EXR::InjectNodesPostEvaluation(const FName& InBranchName, UMovieGraphEvaluatedConfig* InEvaluatedConfig, TArray<UMovieGraphSettingNode*>& OutInjectedNodes)
 {
 	INJECT_BURN_IN_NODE(TEXT("EXR"))
 }
 
+#if WITH_UNREALEXR
 void UMovieGraphImageSequenceOutputNode_MultiLayerEXR::OnReceiveImageDataImpl(UMovieGraphPipeline* InPipeline, UE::MovieGraph::FMovieGraphOutputMergerFrame* InRawFrameData, const TSet<FMovieGraphRenderDataIdentifier>& InMask)
 {
 	auto GetPayloadCompressionType = [](
@@ -766,7 +769,6 @@ void UMovieGraphImageSequenceOutputNode_MultiLayerEXR::OnReceiveImageDataImpl(UM
 		InPipeline->AddOutputFuture(ImageWriteQueue->Enqueue(MoveTemp(MultiLayerImageTask)), OutputFutureData);
 	}
 }
-
 #endif // WITH_UNREALEXR
 
 void UMovieGraphImageSequenceOutputNode_MultiLayerEXR::InjectNodesPostEvaluation(const FName& InBranchName, UMovieGraphEvaluatedConfig* InEvaluatedConfig, TArray<UMovieGraphSettingNode*>& OutInjectedNodes)
