MoviePipeline: Fixed a garbage collection crash that was introduced with the recent burn-in UX overhaul (44669256).

diff --git a/Engine/Plugins/MovieScene/MovieRenderPipeline/Source/MovieRenderPipelineCore/Private/Graph/Nodes/MovieGraphBurnInNode.cpp b/Engine/Plugins/MovieScene/MovieRenderPipeline/Source/MovieRenderPipelineCore/Private/Graph/Nodes/MovieGraphBurnInNode.cpp
--- a/Engine/Plugins/MovieScene/MovieRenderPipeline/Source/MovieRenderPipelineCore/Private/Graph/Nodes/MovieGraphBurnInNode.cpp
+++ b/Engine/Plugins/MovieScene/MovieRenderPipeline/Source/MovieRenderPipelineCore/Private/Graph/Nodes/MovieGraphBurnInNode.cpp
@@ -157,6 +157,13 @@ void UMovieGraphBurnInNode::FMovieGraphBurnInPass::Setup(TWeakObjectPtr<UMovieGr
 		return;
 	}
 	CachedBurnInWidgetClass = ParentNodeOnInitialization->BurnInClass.TryLoadClass<UMovieGraphBurnInWidget>();
+
+	// Output-specific burn-ins keep track of what their parent node's instance name is. There can be multiple output-specific burn-ins
+	// on any given branch, so this is needed in order to disambiguate.
+	if (ParentNodeOnInitialization->IsA<UMovieGraphOutputBurnInNode>())
+	{
+		RendererNodeInstanceName = ParentNodeOnInitialization->GetNodeInstanceName();
+	}
 }
 
 TSharedPtr<SWidget> UMovieGraphBurnInNode::FMovieGraphBurnInPass::GetWidget(UMovieGraphWidgetRendererBaseNode* InNodeThisFrame)
@@ -220,6 +227,30 @@ int32 UMovieGraphBurnInNode::FMovieGraphBurnInPass::GetCompositingSortOrder() co
 	return 100;
 }
 
+UMovieGraphWidgetRendererBaseNode* UMovieGraphBurnInNode::FMovieGraphBurnInPass::GetParentNode(UMovieGraphEvaluatedConfig* InConfig) const
+{
+	// If there's no instance name, fall back to the default behavior.
+	if (RendererNodeInstanceName.IsEmpty())
+	{
+		return FMovieGraphWidgetPass::GetParentNode(InConfig);
+	}
+	
+	constexpr bool bIncludeCDOs = false;
+	TArray<UMovieGraphOutputBurnInNode*> BurnInNodes = InConfig->GetSettingsForBranch<UMovieGraphOutputBurnInNode>(RenderDataIdentifier.RootBranchName, bIncludeCDOs);
+
+	// Find the output-specific burn-in node assigned to this pass.
+	for (UMovieGraphOutputBurnInNode* BurnInNode : BurnInNodes)
+	{
+		if (BurnInNode->GetNodeInstanceName() == RendererNodeInstanceName)
+		{
+			return BurnInNode;
+		}
+	}
+
+	ensureMsgf(false, TEXT("FMovieGraphBurnInPass should not exist without parent node in graph."));
+	return nullptr;
+}
+
 UClass* UMovieGraphBurnInNode::FMovieGraphBurnInPass::GetBurnInClass() const
 {
 	return CachedBurnInWidgetClass.Get();
diff --git a/Engine/Plugins/MovieScene/MovieRenderPipeline/Source/MovieRenderPipelineCore/Private/Graph/Nodes/MovieGraphWidgetRendererBaseNode.cpp b/Engine/Plugins/MovieScene/MovieRenderPipeline/Source/MovieRenderPipelineCore/Private/Graph/Nodes/MovieGraphWidgetRendererBaseNode.cpp
--- a/Engine/Plugins/MovieScene/MovieRenderPipeline/Source/MovieRenderPipelineCore/Private/Graph/Nodes/MovieGraphWidgetRendererBaseNode.cpp
+++ b/Engine/Plugins/MovieScene/MovieRenderPipeline/Source/MovieRenderPipelineCore/Private/Graph/Nodes/MovieGraphWidgetRendererBaseNode.cpp
@@ -119,7 +119,7 @@ UMovieGraphWidgetRendererBaseNode* UMovieGraphWidgetRendererBaseNode::FMovieGrap
 void UMovieGraphWidgetRendererBaseNode::FMovieGraphWidgetPass::Render(const FMovieGraphTraversalContext& InFrameTraversalContext, const FMovieGraphTimeStepData& InTimeData)
 {
 	const UMovieGraphPipeline* Pipeline = Renderer->GetOwningGraph();
-	UMovieGraphWidgetRendererBaseNode* ParentNodeThisFrame = Cast<UMovieGraphWidgetRendererBaseNode>(LayerData.RenderPassNode);
+	UMovieGraphWidgetRendererBaseNode* ParentNodeThisFrame = GetParentNode(InTimeData.EvaluatedConfig);
 
 	if(Pipeline->GetActiveShotList()[Pipeline->GetCurrentShotIndex()]->ShotInfo.State != EMovieRenderShotState::Rendering)
 	{
diff --git a/Engine/Plugins/MovieScene/MovieRenderPipeline/Source/MovieRenderPipelineCore/Public/Graph/Nodes/MovieGraphBurnInNode.h b/Engine/Plugins/MovieScene/MovieRenderPipeline/Source/MovieRenderPipelineCore/Public/Graph/Nodes/MovieGraphBurnInNode.h
--- a/Engine/Plugins/MovieScene/MovieRenderPipeline/Source/MovieRenderPipelineCore/Public/Graph/Nodes/MovieGraphBurnInNode.h
+++ b/Engine/Plugins/MovieScene/MovieRenderPipeline/Source/MovieRenderPipelineCore/Public/Graph/Nodes/MovieGraphBurnInNode.h
@@ -60,6 +60,7 @@ protected:
 		virtual TSharedPtr<SWidget> GetWidget(UMovieGraphWidgetRendererBaseNode* InNodeThisFrame) override;
 		virtual void Render(const FMovieGraphTraversalContext& InFrameTraversalContext, const FMovieGraphTimeStepData& InTimeData) override;
 		virtual int32 GetCompositingSortOrder() const override;
+		virtual UMovieGraphWidgetRendererBaseNode* GetParentNode(UMovieGraphEvaluatedConfig* InConfig) const override;
 
 		/** Gets the burn-in class that this pass will be using. */
 		UClass* GetBurnInClass() const;
@@ -67,6 +68,13 @@ protected:
 	private:
 		TObjectPtr<UMovieGraphBurnInWidget> GetBurnInWidget(UMovieGraphWidgetRendererBaseNode* InNodeThisFrame) const;
 		TSubclassOf<UMovieGraphBurnInWidget> CachedBurnInWidgetClass;
+
+		/**
+		 * If this is an output-specific burn-in (eg, just for JPG), the burn-in renderer node will have a specific instance name. That will
+		 * identify it amongst all of the other output-specific burn-in nodes on the branch (if there are any others). If this is empty, then this
+		 * pass doesn't represent an output-specific burn-in, and is for the legacy burn-in node.
+		 */
+		FString RendererNodeInstanceName;
 	};
 
 private:
