MoviePipeline: Added a new modifier node to the graph to support updating Material Parameter Collections.

diff --git a/Engine/Plugins/MovieScene/MovieRenderPipeline/Source/MovieRenderPipelineCore/Private/Graph/MovieGraphNode.cpp b/Engine/Plugins/MovieScene/MovieRenderPipeline/Source/MovieRenderPipelineCore/Private/Graph/MovieGraphNode.cpp
--- a/Engine/Plugins/MovieScene/MovieRenderPipeline/Source/MovieRenderPipelineCore/Private/Graph/MovieGraphNode.cpp
+++ b/Engine/Plugins/MovieScene/MovieRenderPipeline/Source/MovieRenderPipelineCore/Private/Graph/MovieGraphNode.cpp
@@ -348,6 +348,11 @@ UMovieGraphConfig* UMovieGraphNode::GetGraph() const
 	return Cast<UMovieGraphConfig>(GetOuter());
 }
 
+FInstancedPropertyBag* UMovieGraphNode::GetMutableDynamicProperties_Unsafe() const
+{
+	return const_cast<FInstancedPropertyBag*>(&DynamicProperties);
+}
+
 UMovieGraphPin* UMovieGraphNode::GetInputPin(const FName& InPinLabel, const EMovieGraphPinQueryRequirement PinRequirement) const
 {
 	for (UMovieGraphPin* InputPin : InputPins)
diff --git a/Engine/Plugins/MovieScene/MovieRenderPipeline/Source/MovieRenderPipelineCore/Private/Graph/MovieGraphPipeline.cpp b/Engine/Plugins/MovieScene/MovieRenderPipeline/Source/MovieRenderPipelineCore/Private/Graph/MovieGraphPipeline.cpp
--- a/Engine/Plugins/MovieScene/MovieRenderPipeline/Source/MovieRenderPipelineCore/Private/Graph/MovieGraphPipeline.cpp
+++ b/Engine/Plugins/MovieScene/MovieRenderPipeline/Source/MovieRenderPipelineCore/Private/Graph/MovieGraphPipeline.cpp
@@ -362,8 +362,8 @@ void UMovieGraphPipeline::UpdateLayerContentsInRenderLayerSubsystem(const UMovie
 		constexpr bool bExactMatch = true;
 		const TArray<UMovieGraphCollectionNode*> CollectionNodes =
 			EvaluatedConfig->GetSettingsForBranch<UMovieGraphCollectionNode>(LayerName, bIncludeCDOs, bExactMatch);
-		TArray<UMovieGraphModifierNode*> ModifierNodes =
-			EvaluatedConfig->GetSettingsForBranch<UMovieGraphModifierNode>(LayerName, bIncludeCDOs, bExactMatch);
+		TArray<IMovieGraphModifierNodeInterface*> ModifierNodes =
+			EvaluatedConfig->GetSettingsImplementing<IMovieGraphModifierNodeInterface>(UMovieGraphModifierNodeInterface::StaticClass(), LayerName);
 
 		// Graph evaluation discovers nodes working from the Outputs node, moving towards the Inputs node. Therefore modifiers (even if they are
 		// nodes which override an existing modifier) that are furthest downstream occur first in the array of modifiers returned by the evaluated
@@ -377,15 +377,33 @@ void UMovieGraphPipeline::UpdateLayerContentsInRenderLayerSubsystem(const UMovie
 		// Therefore, the modifier node array returned by the evaluated graph needs to be reversed.
 		Algo::Reverse(ModifierNodes);
 		
-		for (const UMovieGraphModifierNode* ModifierNode : ModifierNodes)
+		for (const IMovieGraphModifierNodeInterface* ModifierNode : ModifierNodes)
 		{
+			// For modifiers that don't use collections, just add all of the modifiers unconditionally.
+			if (!ModifierNode->SupportsCollections())
+			{
+				for (UMovieGraphModifierBase* Modifier : ModifierNode->GetAllModifiers())
+				{
+					RenderLayer->AddLayerModifier(Modifier);
+				}
+				
+				continue;
+			}
+			
 			// For each modifier instance, find the collection(s) that it is modifying
-			for (UMovieGraphCollectionModifier* ModifierInstance : ModifierNode->GetModifiers())
+			for (UMovieGraphModifierBase* ModifierInstance : ModifierNode->GetAllModifiers())
 			{
-				ModifierInstance->SetCollections({});
+				// All of these modifiers should support collections (due to the SupportsCollections() call), but double-check anyway
+				UMovieGraphCollectionModifier* CollectionModifierInstance = Cast<UMovieGraphCollectionModifier>(ModifierInstance);
+				if (!ensureMsgf(CollectionModifierInstance, TEXT("Found a modifier node which declares it supports collections, but provided modifiers that don't use collections.")))
+				{
+					continue;
+				}
+				
+				CollectionModifierInstance->SetCollections({});
 				TArray<UMovieGraphCollection*> ModifierCollections;
 
-				for (const FName& ModifiedCollectionName : ModifierNode->GetCollections())
+				for (const FName& ModifiedCollectionName : ModifierNode->GetAllCollections())
 				{
 					bool bFoundModifiedCollection = false;
 					
@@ -413,7 +431,7 @@ void UMovieGraphPipeline::UpdateLayerContentsInRenderLayerSubsystem(const UMovie
 					if (!bFoundModifiedCollection)
 					{
 						UE_LOG(LogMovieRenderPipeline, Warning, TEXT("The modifier '%s' specified a collection '%s', but the collection couldn't be found."),
-							*ModifierNode->ModifierName, *ModifiedCollectionName.ToString());
+							*CollectionModifierInstance->GetName(), *ModifiedCollectionName.ToString());
 					}
 				}
 
@@ -421,9 +439,9 @@ void UMovieGraphPipeline::UpdateLayerContentsInRenderLayerSubsystem(const UMovie
 				// won't have any effect without collection(s) to act on.
 				if (!ModifierCollections.IsEmpty())
 				{
-					ModifierInstance->SetCollections(ModifierCollections);
+					CollectionModifierInstance->SetCollections(ModifierCollections);
 
-					RenderLayer->AddModifier(ModifierInstance);
+					RenderLayer->AddLayerModifier(CollectionModifierInstance);
 				}
 			}
 		}
diff --git a/Engine/Plugins/MovieScene/MovieRenderPipeline/Source/MovieRenderPipelineCore/Private/Graph/MovieGraphRenderLayerSubsystem.cpp b/Engine/Plugins/MovieScene/MovieRenderPipeline/Source/MovieRenderPipelineCore/Private/Graph/MovieGraphRenderLayerSubsystem.cpp
--- a/Engine/Plugins/MovieScene/MovieRenderPipeline/Source/MovieRenderPipelineCore/Private/Graph/MovieGraphRenderLayerSubsystem.cpp
+++ b/Engine/Plugins/MovieScene/MovieRenderPipeline/Source/MovieRenderPipelineCore/Private/Graph/MovieGraphRenderLayerSubsystem.cpp
@@ -499,6 +499,16 @@ void UMovieGraphCollectionModifier::AddCollection(UMovieGraphCollection* Collect
 	Collections.Add(Collection);
 }
 
+void UMovieGraphCollectionModifier::SetCollections(const TArray<UMovieGraphCollection*> InCollections)
+{
+	Collections = InCollections;
+}
+
+TArray<UMovieGraphCollection*> UMovieGraphCollectionModifier::GetCollections() const
+{
+	return Collections;
+}
+
 UMovieGraphConditionGroupQueryBase::UMovieGraphConditionGroupQueryBase()
 	: OpType(EMovieGraphConditionGroupQueryOpType::Add)
 	, bIsEnabled(true)
@@ -3553,14 +3563,15 @@ const FString& UMovieGraphCollection::GetCollectionName() const
 
 UMovieGraphCollection* UMovieGraphRenderLayer::GetCollectionByName(const FString& Name) const
 {
-	for (const UMovieGraphCollectionModifier* Modifier : Modifiers)
+	for (const UMovieGraphModifierBase* Modifier : Modifiers)
 	{
-		if (!Modifier)
+		const UMovieGraphCollectionModifier* CollectionModifier = Cast<UMovieGraphCollectionModifier>(Modifier);
+		if (!CollectionModifier)
 		{
 			continue;
 		}
 		
-		for (UMovieGraphCollection* Collection : Modifier->GetCollections())
+		for (UMovieGraphCollection* Collection : CollectionModifier->GetCollections())
 		{
 			if (Collection && Collection->GetCollectionName().Equals(Name))
 			{
@@ -3580,11 +3591,43 @@ void UMovieGraphRenderLayer::AddModifier(UMovieGraphCollectionModifier* Modifier
 	}
 }
 
+void UMovieGraphRenderLayer::AddLayerModifier(UMovieGraphModifierBase* InModifier)
+{
+	if (!Modifiers.Contains(InModifier))
+	{
+		Modifiers.Add(InModifier);
+	}
+}
+
+TArray<UMovieGraphCollectionModifier*> UMovieGraphRenderLayer::GetModifiers() const
+{
+	TArray<UMovieGraphCollectionModifier*> ConvertedModifiers;
+	for (const TObjectPtr<UMovieGraphModifierBase>& Modifier : Modifiers)
+	{
+		if (UMovieGraphCollectionModifier* CollectionModifier = Cast<UMovieGraphCollectionModifier>(Modifier))
+		{
+			ConvertedModifiers.Add(CollectionModifier);
+		}
+	}
+		
+	return ConvertedModifiers;
+}
+
+TArray<UMovieGraphModifierBase*> UMovieGraphRenderLayer::GetLayerModifiers() const
+{
+	return Modifiers;
+}
+
 void UMovieGraphRenderLayer::RemoveModifier(UMovieGraphCollectionModifier* Modifier)
 {
 	Modifiers.Remove(Modifier);
 }
 
+void UMovieGraphRenderLayer::RemoveLayerModifier(UMovieGraphModifierBase* Modifier)
+{
+	Modifiers.Remove(Modifier);
+}
+
 void UMovieGraphRenderLayer::Apply(const UWorld* World)
 {
 	if (!World)
@@ -3593,7 +3636,7 @@ void UMovieGraphRenderLayer::Apply(const UWorld* World)
 	}
 	
 	// Apply all modifiers
-	for (UMovieGraphCollectionModifier* Modifier : Modifiers)
+	for (UMovieGraphModifierBase* Modifier : Modifiers)
 	{
 		Modifier->ApplyModifier(World);
 	}
@@ -3605,14 +3648,13 @@ void UMovieGraphRenderLayer::Revert()
 	// state of one modifier may depend on modifiers that were previously applied.
 	for (int32 Index = Modifiers.Num() - 1; Index >= 0; Index--)
 	{
-		if (UMovieGraphCollectionModifier* Modifier = Modifiers[Index])
+		if (UMovieGraphModifierBase* Modifier = Modifiers[Index])
 		{
 			Modifier->UndoModifier();
 		}
 	}
 }
 
-
 UMovieGraphRenderLayerSubsystem* UMovieGraphRenderLayerSubsystem::GetFromWorld(const UWorld* World)
 {
 	if (World)
diff --git a/Engine/Plugins/MovieScene/MovieRenderPipeline/Source/MovieRenderPipelineCore/Private/Graph/Nodes/MovieGraphMaterialParameterCollectionModifier.cpp b/Engine/Plugins/MovieScene/MovieRenderPipeline/Source/MovieRenderPipelineCore/Private/Graph/Nodes/MovieGraphMaterialParameterCollectionModifier.cpp
new file mode 100644
--- /dev/null
+++ b/Engine/Plugins/MovieScene/MovieRenderPipeline/Source/MovieRenderPipelineCore/Private/Graph/Nodes/MovieGraphMaterialParameterCollectionModifier.cpp
@@ -0,0 +1,81 @@
+﻿// Copyright Epic Games, Inc. All Rights Reserved.
+
+#include "Graph/Nodes/MovieGraphMaterialParameterCollectionModifier.h"
+
+#include "Engine/World.h"
+#include "Materials/MaterialParameterCollection.h"
+#include "Materials/MaterialParameterCollectionInstance.h"
+
+#include UE_INLINE_GENERATED_CPP_BY_NAME(MovieGraphMaterialParameterCollectionModifier)
+
+void UMovieGraphMaterialParameterCollectionModifier::ApplyModifier(const UWorld* World)
+{
+	CachedWorld = World;
+	
+	constexpr bool bCacheCurrentValues = true;
+	ApplyScalarAndVectorValues(ScalarParameterUpdates, VectorParameterUpdates, bCacheCurrentValues);
+}
+
+void UMovieGraphMaterialParameterCollectionModifier::UndoModifier()
+{
+	ApplyScalarAndVectorValues(PriorScalarValues, PriorVectorValues);
+	
+	PriorScalarValues.Reset();
+	PriorVectorValues.Reset();
+}
+
+void UMovieGraphMaterialParameterCollectionModifier::ClearParameterValues()
+{
+	ScalarParameterUpdates.Reset();
+	VectorParameterUpdates.Reset();
+}
+
+void UMovieGraphMaterialParameterCollectionModifier::ApplyScalarAndVectorValues(const TMap<FName, float>& InScalarValues, const TMap<FName, FLinearColor>& InVectorValues, const bool bCacheCurrentValues)
+{
+	if (!CachedWorld.IsValid())
+	{
+		return;
+	}
+	
+	const UMaterialParameterCollection* MPC = MaterialParameterCollection.LoadSynchronous();
+	if (!MPC)
+	{
+		return;
+	}
+	
+	UMaterialParameterCollectionInstance* MpcInstance = CachedWorld->GetParameterCollectionInstance(MPC);
+	if (!MpcInstance)
+	{
+		return;
+	}
+
+	// Cache + apply scalar parameters
+	for (const TPair<FName, float>& ScalarPair : InScalarValues)
+	{
+		if (bCacheCurrentValues)
+		{
+			float ScalarParamValue = 0.f;
+			if (MpcInstance->GetScalarParameterValue(ScalarPair.Key, ScalarParamValue))
+			{
+				PriorScalarValues.Add(ScalarPair.Key, ScalarParamValue);
+			}
+		}
+
+		MpcInstance->SetScalarParameterValue(ScalarPair.Key, ScalarPair.Value);
+	}
+
+	// Cache + apply vector parameters
+	for (const TPair<FName, FLinearColor>& VectorPair : InVectorValues)
+	{
+		if (bCacheCurrentValues)
+		{
+			FLinearColor VectorParamValue;
+			if (MpcInstance->GetVectorParameterValue(VectorPair.Key, VectorParamValue))
+			{
+				PriorVectorValues.Add(VectorPair.Key, VectorParamValue);
+			}
+		}
+
+		MpcInstance->SetVectorParameterValue(VectorPair.Key, VectorPair.Value);
+	}
+}
diff --git a/Engine/Plugins/MovieScene/MovieRenderPipeline/Source/MovieRenderPipelineCore/Private/Graph/Nodes/MovieGraphMaterialParameterCollectionModifierNode.cpp b/Engine/Plugins/MovieScene/MovieRenderPipeline/Source/MovieRenderPipelineCore/Private/Graph/Nodes/MovieGraphMaterialParameterCollectionModifierNode.cpp
new file mode 100644
--- /dev/null
+++ b/Engine/Plugins/MovieScene/MovieRenderPipeline/Source/MovieRenderPipelineCore/Private/Graph/Nodes/MovieGraphMaterialParameterCollectionModifierNode.cpp
@@ -0,0 +1,307 @@
+// Copyright Epic Games, Inc. All Rights Reserved.
+
+#include "Graph/Nodes/MovieGraphMaterialParameterCollectionModifierNode.h"
+
+#include "Graph/Nodes/MovieGraphMaterialParameterCollectionModifier.h"
+#include "Materials/MaterialParameterCollection.h"
+
+#include UE_INLINE_GENERATED_CPP_BY_NAME(MovieGraphMaterialParameterCollectionModifierNode)
+
+#define LOCTEXT_NAMESPACE "MovieGraph"
+
+UMovieGraphMaterialParameterCollectionModifierNode::UMovieGraphMaterialParameterCollectionModifierNode()
+{
+	Modifier = CreateDefaultSubobject<UMovieGraphMaterialParameterCollectionModifier>(TEXT("MPC_Modifier"));
+}
+
+#if WITH_EDITOR
+FText UMovieGraphMaterialParameterCollectionModifierNode::GetNodeTitle(const bool bGetDescriptive) const
+{
+	static const FText NodeName = LOCTEXT("NodeName_MPCModifier", "Material Parameter Collection Modifier");
+	static const FText NodeDescription = LOCTEXT("NodeDescription_MPCModifier", "Material Parameter Collection Modifier\n{0}");
+
+	const UMaterialParameterCollection* MPC = MaterialParameterCollection.LoadSynchronous(); 
+
+	const FString AssetName = IsValid(MPC) ? MPC->GetName() : FString();
+	const FString MpcDisplayName = !AssetName.IsEmpty() && bOverride_MaterialParameterCollection
+		? AssetName
+		: LOCTEXT("NodeNoNameWarning_MPCModifier", "No Material Parameter Collection Asset Chosen").ToString();
+	
+	if (bGetDescriptive)
+	{
+		return FText::Format(NodeDescription, FText::FromString(MpcDisplayName));
+	}
+
+	return NodeName;
+}
+
+FText UMovieGraphMaterialParameterCollectionModifierNode::GetMenuCategory() const
+{
+	return LOCTEXT("CollectionNode_Category", "Utility");
+}
+
+FLinearColor UMovieGraphMaterialParameterCollectionModifierNode::GetNodeTitleColor() const
+{
+	static const FLinearColor ModifierNodeColor = FLinearColor(0.6f, 0.113f, 0.113f);
+	return ModifierNodeColor;
+}
+
+FSlateIcon UMovieGraphMaterialParameterCollectionModifierNode::GetIconAndTint(FLinearColor& OutColor) const
+{
+	static const FSlateIcon ModifierIcon = FSlateIcon(FAppStyle::GetAppStyleSetName(), "Icons.Edit");
+
+	OutColor = FLinearColor::White;
+	return ModifierIcon;
+}
+
+void UMovieGraphMaterialParameterCollectionModifierNode::PostEditChangeProperty(FPropertyChangedEvent& PropertyChangedEvent)
+{
+	Super::PostEditChangeProperty(PropertyChangedEvent);
+	
+	if (!HasAnyFlags(RF_ClassDefaultObject))
+	{
+		// Update the node and its dynamic properties if the MPC asset changes.
+		if ((PropertyChangedEvent.GetPropertyName() == GET_MEMBER_NAME_CHECKED(UMovieGraphMaterialParameterCollectionModifierNode, MaterialParameterCollection)) ||
+			(PropertyChangedEvent.GetPropertyName() == GET_MEMBER_NAME_CHECKED(UMovieGraphMaterialParameterCollectionModifierNode, bOverride_MaterialParameterCollection)))
+		{
+			UpdateDynamicProperties();
+			OnNodeChangedDelegate.Broadcast(this);
+		}
+	}
+}
+
+void UMovieGraphMaterialParameterCollectionModifierNode::PostLoad()
+{
+	Super::PostLoad();
+
+	if (!HasAnyFlags(RF_ClassDefaultObject))
+	{
+		// Update the node's dynamic properties if any properties within the MPC asset changes (its scalar/vector parameters).
+		FCoreUObjectDelegates::OnObjectPropertyChanged.AddUObject(this, &UMovieGraphMaterialParameterCollectionModifierNode::OnObjectPropertyChanged);
+		UpdateDynamicProperties();
+	}
+}
+#endif // WITH_EDITOR
+
+TArray<FPropertyBagPropertyDesc> UMovieGraphMaterialParameterCollectionModifierNode::GetDynamicPropertyDescriptions() const
+{
+	TArray<FPropertyBagPropertyDesc> ModifierDescs;
+	
+	UMaterialParameterCollection* MPC = MaterialParameterCollection.LoadSynchronous();
+	if (!IsValid(MPC))
+	{
+		return ModifierDescs;
+	}
+
+	// Creates a new property bag property description for a specific MPC parameter.
+	auto CreatePropertyDescFromMpcParameter = [&ModifierDescs](
+		UMaterialParameterCollection* InMPC, const FCollectionParameterBase& InParameter,
+		const EPropertyBagPropertyType ParamType, const UObject* ParamTypeObject = nullptr)
+	{
+		// The parameter name may contain spaces or other invalid characters for property names, so we have to sanitize it. Unfortunately this leaves
+		// open a potential edge case where the MPC could contain two parameters (eg, "Param With Spaces" and "Param_With_Spaces") that would resolve
+		// to the same sanitized name.
+		const FName ParamName = FInstancedPropertyBag::SanitizePropertyName(InParameter.ParameterName, TEXT('_'));
+		const FString OverrideName = FString::Printf(TEXT("bOverride_%s"), *ParamName.ToString());
+
+		FPropertyBagPropertyDesc NewDesc(ParamName, ParamType, ParamTypeObject);
+		NewDesc.ID = InParameter.Id;
+#if WITH_EDITOR
+		NewDesc.MetaData.Add(FPropertyBagPropertyDescMetaData(CategoryMetadataKey, InMPC->GetName()));
+		NewDesc.MetaData.Add(FPropertyBagPropertyDescMetaData(EditConditionMetadataKey, OverrideName));
+		NewDesc.MetaData.Add(FPropertyBagPropertyDescMetaData(DisplayNameMetadataKey, InParameter.ParameterName.ToString()));
+#endif	// WITH_EDITOR
+		ModifierDescs.Add(MoveTemp(NewDesc));
+
+		// Need a bOverride_* property in order for this to be picked up properly by graph evaluation. Note that the ID of the the override property
+		// is based off of the ID of the parameter. The overrides don't exist in the MPC asset, so we have to come up with new IDs for the overrides.
+		// The easiest way to do this is to just use the string of the parameter's ID to create a new deterministic FGuid for the override.
+		FPropertyBagPropertyDesc OverrideDesc(FName(OverrideName), EPropertyBagPropertyType::Bool);
+		OverrideDesc.ID = FGuid::NewDeterministicGuid(NewDesc.ID.ToString());
+#if WITH_EDITOR
+		OverrideDesc.MetaData.Add(FPropertyBagPropertyDescMetaData(InlineEditConditionToggleMetadataKey, FString(TEXT("true"))));
+#endif	// WITH_EDITOR
+		ModifierDescs.Add(MoveTemp(OverrideDesc));
+	};
+
+	// Create property descs for each MPC in the MPC "chain"
+	UMaterialParameterCollection* CurrentMPC = MPC;
+	while (CurrentMPC)
+	{
+		for (const FCollectionScalarParameter& MPCScalarParameter : CurrentMPC->ScalarParameters)
+		{
+			CreatePropertyDescFromMpcParameter(CurrentMPC, MPCScalarParameter, EPropertyBagPropertyType::Float);
+		}
+
+		for (const FCollectionVectorParameter& MPCVectorParameter : CurrentMPC->VectorParameters)
+		{
+			CreatePropertyDescFromMpcParameter(CurrentMPC, MPCVectorParameter, EPropertyBagPropertyType::Struct, TBaseStructure<FLinearColor>::Get());
+		}
+
+		CurrentMPC = CurrentMPC->GetBaseParameterCollection();
+	}
+
+	return ModifierDescs;
+}
+
+bool UMovieGraphMaterialParameterCollectionModifierNode::GetDynamicPropertyValue(const FName PropertyName, FString& OutValue)
+{
+	TValueOrError<FString, EPropertyBagResult> Result = DynamicProperties.GetValueSerializedString(PropertyName);
+	if (Result.HasValue())
+	{
+		OutValue = Result.GetValue();
+		return true;
+	}
+
+	return false;
+}
+
+TArray<FMovieGraphPropertyInfo> UMovieGraphMaterialParameterCollectionModifierNode::GetOverrideablePropertyInfo() const
+{
+	TArray<FMovieGraphPropertyInfo> PropertyInfo = Super::GetOverrideablePropertyInfo();
+	
+	// Don't allow the Material Parameter Collection asset property to be overridden. The individual MPC parameters can, but not the asset itself.
+	PropertyInfo.RemoveAll([](const FMovieGraphPropertyInfo& InPropertyInfo)
+	{
+		return InPropertyInfo.Name == GET_MEMBER_NAME_CHECKED(UMovieGraphMaterialParameterCollectionModifierNode, MaterialParameterCollection);
+	});
+
+	return PropertyInfo;
+}
+
+TArray<UMovieGraphModifierBase*> UMovieGraphMaterialParameterCollectionModifierNode::GetAllModifiers() const
+{
+	// Update the modifier with the scalar/vector parameters set on the node.
+
+	Modifier->ClearParameterValues();
+	Modifier->MaterialParameterCollection = MaterialParameterCollection;
+
+	if (!MaterialParameterCollection.LoadSynchronous())
+	{
+		return {};
+	}
+
+	// First, figure out which parameters have actually been overridden on the node. Only these parameter values should be set on the modifier.
+	TSet<FName> OverriddenParameterNames;
+	for (const FPropertyBagPropertyDesc& Desc : GetDynamicPropertyDescriptions())
+	{
+		// Determine if this is a bOverride_* property. The only properties in the bag that are bools are override properties.
+		const bool bIsOverrideProperty = (Desc.ValueType == EPropertyBagPropertyType::Bool);
+
+		// Regular property, not a bOverride_*
+		if (!bIsOverrideProperty)
+		{
+			continue;
+		}
+
+		const TValueOrError<bool, EPropertyBagResult> OverrideResult = DynamicProperties.GetValueBool(Desc.Name);
+		if (OverrideResult.HasValue() && (OverrideResult.GetValue() == true))
+		{
+			// Generate the parameter name from the override's property name. Unfortunately we cannot rely on metadata here to link the parameter's
+			// property and the override's property (since metadata is not available in non-editor builds). The next-best way to link them is via
+			// name.
+			const FName ParameterName = FName(Desc.Name.ToString().Replace(TEXT("bOverride_"), TEXT(""), ESearchCase::CaseSensitive));
+			
+			OverriddenParameterNames.Add(ParameterName);
+		}
+	}
+
+	// Second, for all overridden parameters, transfer over their values to the modifier. 
+	for (const FPropertyBagPropertyDesc& Desc : GetDynamicPropertyDescriptions())
+	{
+		const EPropertyBagPropertyType ParamType = Desc.ValueType;
+		const FName ParamName = Desc.Name;
+
+		if (!OverriddenParameterNames.Contains(ParamName))
+		{
+			continue;
+		}
+
+		// Note: When adding a new entry to ScalarParameterUpdates/VectorParameterUpdates, note that GetParameterName() is used *rather than* the
+		// Desc's name. The parameter name may contain spaces or other property-name-incompatible characters, so we need to use the desc's ID (which
+		// is the param's ID) to fetch its param name properly. Unfortunately MPC parameters can only be set by name, not ID.
+		const FName ActualParamName = MaterialParameterCollection->GetParameterName(Desc.ID);
+		if (ActualParamName.IsNone())
+		{
+			continue;
+		}
+		
+		if (ParamType == EPropertyBagPropertyType::Float)
+		{
+			const TValueOrError<float, EPropertyBagResult> FloatResult = DynamicProperties.GetValueFloat(ParamName);
+			if (FloatResult.HasValue())
+			{
+				Modifier->ScalarParameterUpdates.Add(ActualParamName, FloatResult.GetValue());
+			}
+		}
+		else if (ParamType == EPropertyBagPropertyType::Struct)
+		{
+			const TValueOrError<FLinearColor*, EPropertyBagResult> StructResult = DynamicProperties.GetValueStruct<FLinearColor>(ParamName);
+			if (StructResult.HasValue())
+			{
+				Modifier->VectorParameterUpdates.Add(ActualParamName, *StructResult.GetValue());
+			}
+		}
+		else
+		{
+			// Invalid parameter type	
+			check(false);
+		}
+	}
+
+	return { Modifier.Get() };
+}
+
+bool UMovieGraphMaterialParameterCollectionModifierNode::SupportsCollections() const
+{
+	return false;
+}
+
+FString UMovieGraphMaterialParameterCollectionModifierNode::GetNodeInstanceName() const
+{
+	if (!MaterialParameterCollection.IsNull())
+	{
+		return MaterialParameterCollection.GetAssetName();
+	}
+
+	return FString();
+}
+
+void UMovieGraphMaterialParameterCollectionModifierNode::PrepareForFlattening(const UMovieGraphSettingNode* InSourceNode)
+{
+	Super::PrepareForFlattening(InSourceNode);
+
+	// This node's dynamic properties rely on the MPC being valid and correct in order to be properly initialized.
+	if (const UMovieGraphMaterialParameterCollectionModifierNode* SourcePresetNode = Cast<UMovieGraphMaterialParameterCollectionModifierNode>(InSourceNode))
+	{
+		MaterialParameterCollection = SourcePresetNode->MaterialParameterCollection;
+	}
+}
+
+void UMovieGraphMaterialParameterCollectionModifierNode::OnObjectPropertyChanged(UObject* ObjectBeingModified, FPropertyChangedEvent& PropertyChangedEvent)
+{
+	// Since this method can be called at high frequency, early-out if possible
+	if ((PropertyChangedEvent.ChangeType == EPropertyChangeType::Interactive)
+		|| !ObjectBeingModified
+		|| (ObjectBeingModified->GetClass() != UMaterialParameterCollection::StaticClass()))
+	{
+		return;
+	}
+	
+	// Gather all of the MPCs in the MPC "chain"
+	UMaterialParameterCollection* CurrentMPC = MaterialParameterCollection.LoadSynchronous();
+	TArray<UObject*> RelevantMPCs;
+	while (CurrentMPC)
+	{
+		RelevantMPCs.Add(CurrentMPC);
+		CurrentMPC = CurrentMPC->GetBaseParameterCollection();
+	}
+	
+	// If a scalar/vector parameter in the MPC changes, we need to update our property bag to reflect the changes
+	if (RelevantMPCs.Contains(ObjectBeingModified))
+	{
+		UpdateDynamicProperties();
+	}
+}
+
+#undef LOCTEXT_NAMESPACE // "MovieGraph"
diff --git a/Engine/Plugins/MovieScene/MovieRenderPipeline/Source/MovieRenderPipelineCore/Private/Graph/Nodes/MovieGraphModifierNode.cpp b/Engine/Plugins/MovieScene/MovieRenderPipeline/Source/MovieRenderPipelineCore/Private/Graph/Nodes/MovieGraphModifierNode.cpp
--- a/Engine/Plugins/MovieScene/MovieRenderPipeline/Source/MovieRenderPipelineCore/Private/Graph/Nodes/MovieGraphModifierNode.cpp
+++ b/Engine/Plugins/MovieScene/MovieRenderPipeline/Source/MovieRenderPipelineCore/Private/Graph/Nodes/MovieGraphModifierNode.cpp
@@ -153,6 +153,11 @@ void UMovieGraphModifierNode::PostEditChangeProperty(FPropertyChangedEvent& Prop
 }
 #endif // WITH_EDITOR
 
+FString UMovieGraphModifierNode::GetNodeInstanceName() const
+{
+	return ModifierName;
+}
+
 UMovieGraphCollectionModifier* UMovieGraphModifierNode::GetModifier(TSubclassOf<UMovieGraphCollectionModifier> ModifierType) const
 {
 	const TObjectPtr<UMovieGraphCollectionModifier>* FoundModifier =
@@ -169,6 +174,21 @@ const TArray<UMovieGraphCollectionModifier*>& UMovieGraphModifierNode::GetModifi
 	return ModifiersContainer->Modifiers;
 }
 
+TArray<UMovieGraphModifierBase*> UMovieGraphModifierNode::GetAllModifiers() const
+{
+	PRAGMA_DISABLE_DEPRECATION_WARNINGS
+	const TArray<UMovieGraphCollectionModifier*>& Modifiers = GetModifiers();
+	PRAGMA_ENABLE_DEPRECATION_WARNINGS
+	
+	TArray<UMovieGraphModifierBase*> BaseModifiers;
+	for (UMovieGraphCollectionModifier* CollectionModifier : Modifiers)
+	{
+		BaseModifiers.Add(Cast<UMovieGraphModifierBase>(CollectionModifier));
+	}
+	
+	return BaseModifiers;
+}
+
 UMovieGraphCollectionModifier* UMovieGraphModifierNode::AddModifier(TSubclassOf<UMovieGraphCollectionModifier> ModifierType)
 {
 	const TObjectPtr<UMovieGraphCollectionModifier>* FoundModifier =
@@ -243,6 +263,18 @@ const TArray<FName>& UMovieGraphModifierNode::GetCollections() const
 	return Collections;
 }
 
+TArray<FName> UMovieGraphModifierNode::GetAllCollections() const
+{
+	PRAGMA_DISABLE_DEPRECATION_WARNINGS
+	return GetCollections();
+	PRAGMA_ENABLE_DEPRECATION_WARNINGS
+}
+
+bool UMovieGraphModifierNode::SupportsCollections() const
+{
+	return true;
+}
+
 void UMovieGraphModifierNode::SetCollectionEnabled(const FName& InCollectionName, const bool bIsCollectionEnabled)
 {
 #if WITH_EDITOR
diff --git a/Engine/Plugins/MovieScene/MovieRenderPipeline/Source/MovieRenderPipelineCore/Public/Graph/MovieGraphNode.h b/Engine/Plugins/MovieScene/MovieRenderPipeline/Source/MovieRenderPipelineCore/Public/Graph/MovieGraphNode.h
--- a/Engine/Plugins/MovieScene/MovieRenderPipeline/Source/MovieRenderPipelineCore/Public/Graph/MovieGraphNode.h
+++ b/Engine/Plugins/MovieScene/MovieRenderPipeline/Source/MovieRenderPipelineCore/Public/Graph/MovieGraphNode.h
@@ -265,6 +265,14 @@ public:
 	MOVIERENDERPIPELINECORE_API void UpdateDynamicProperties();
 	MOVIERENDERPIPELINECORE_API class UMovieGraphConfig* GetGraph() const;
 
+	/**
+	 * Gets the property bag instance that backs the dynamic properties on this node.
+	 * 
+	 * Should not be used except for very specific, niche cases where the property bag needs to be accessed/modified directly, like
+	 * for details customizations. For normal uses, the other dynamic property methods should be used.
+	 */
+	MOVIERENDERPIPELINECORE_API FInstancedPropertyBag* GetMutableDynamicProperties_Unsafe() const;
+
 	/**
 	 * Gets the input pin with the specified name, or nullptr if one could not be found. Most pins on a node are
 	 * "built-in", meaning they ship with the node. Dynamic pins (pins which are not built-in) can potentially have the
diff --git a/Engine/Plugins/MovieScene/MovieRenderPipeline/Source/MovieRenderPipelineCore/Public/Graph/MovieGraphRenderLayerSubsystem.h b/Engine/Plugins/MovieScene/MovieRenderPipeline/Source/MovieRenderPipelineCore/Public/Graph/MovieGraphRenderLayerSubsystem.h
--- a/Engine/Plugins/MovieScene/MovieRenderPipeline/Source/MovieRenderPipelineCore/Public/Graph/MovieGraphRenderLayerSubsystem.h
+++ b/Engine/Plugins/MovieScene/MovieRenderPipeline/Source/MovieRenderPipelineCore/Public/Graph/MovieGraphRenderLayerSubsystem.h
@@ -873,10 +873,23 @@ private:
 };
 
 /**
- * Base class for providing actor modification functionality via collections.
+ * The base class that all modifiers should inherit from. Does not support collections (see UMovieGraphCollectionModifier).
  */
 UCLASS(MinimalAPI, Abstract)
-class UMovieGraphCollectionModifier : public UObject
+class UMovieGraphModifierBase : public UObject
+{
+	GENERATED_BODY()
+
+public:
+	MOVIERENDERPIPELINECORE_API virtual void ApplyModifier(const UWorld* World) PURE_VIRTUAL(UMovieGraphModifierBase::ApplyModifier, );
+	MOVIERENDERPIPELINECORE_API virtual void UndoModifier() PURE_VIRTUAL(UMovieGraphModifierBase::UndoModifier, );
+};
+
+/**
+ * Base class for providing actor modification functionality via collections. For modifiers without collections, see UMovieGraphModifierBase.
+ */
+UCLASS(MinimalAPI, Abstract)
+class UMovieGraphCollectionModifier : public UMovieGraphModifierBase
 {
 	GENERATED_BODY()
 
@@ -887,13 +900,11 @@ public:
 
 	/** Overwrites the existing collections with the provided array of collections. */
 	UFUNCTION(BlueprintCallable, Category = "Settings")
-	void SetCollections(const TArray<UMovieGraphCollection*> InCollections) { Collections = InCollections; }
+	MOVIERENDERPIPELINECORE_API void SetCollections(const TArray<UMovieGraphCollection*> InCollections);
 	
+	/** Gets all collections that this modifier is using. */
 	UFUNCTION(BlueprintCallable, Category = "Settings")
-	TArray<UMovieGraphCollection*> GetCollections() const { return Collections; }
-
-	virtual void ApplyModifier(const UWorld* World) PURE_VIRTUAL(UMovieGraphCollectionModifier::ApplyModifier, );
-	virtual void UndoModifier() PURE_VIRTUAL(UMovieGraphCollectionModifier::UndoModifier, );
+	MOVIERENDERPIPELINECORE_API TArray<UMovieGraphCollection*> GetCollections() const;
 
 protected:
 	/** The collections which this modifier will operate on. */
@@ -1078,15 +1089,27 @@ public:
 	UFUNCTION(BlueprintCallable, Category = "Settings")
 	MOVIERENDERPIPELINECORE_API UMovieGraphCollection* GetCollectionByName(const FString& Name) const;
 
-	UFUNCTION(BlueprintCallable, Category = "Settings")
+	UE_DEPRECATED(5.7, "Use AddLayerModifier() instead.")
+	UFUNCTION(BlueprintCallable, Category = "Settings", meta = (DeprecatedFunction, DeprecationMessage = "Use AddLayerModifier() instead."))
 	MOVIERENDERPIPELINECORE_API void AddModifier(UMovieGraphCollectionModifier* Modifier);
 
 	UFUNCTION(BlueprintCallable, Category = "Settings")
-	TArray<UMovieGraphCollectionModifier*> GetModifiers() const { return Modifiers; };
+	MOVIERENDERPIPELINECORE_API void AddLayerModifier(UMovieGraphModifierBase* InModifier);
+
+	UE_DEPRECATED(5.7, "Use GetLayerModifiers() instead.")
+	UFUNCTION(BlueprintCallable, Category = "Settings", meta = (DeprecatedFunction, DeprecationMessage = "Use GetLayerModifiers() instead."))
+	TArray<UMovieGraphCollectionModifier*> GetModifiers() const;
 
 	UFUNCTION(BlueprintCallable, Category = "Settings")
+	TArray<UMovieGraphModifierBase*> GetLayerModifiers() const;
+
+	UE_DEPRECATED(5.7, "Use RemoveLayerModifier() instead.")
+	UFUNCTION(BlueprintCallable, Category = "Settings", meta = (DeprecatedFunction, DeprecationMessage = "Use RemoveLayerModifier() instead."))
 	MOVIERENDERPIPELINECORE_API void RemoveModifier(UMovieGraphCollectionModifier* Modifier);
 
+	UFUNCTION(BlueprintCallable, Category = "Settings")
+	MOVIERENDERPIPELINECORE_API void RemoveLayerModifier(UMovieGraphModifierBase* Modifier);
+
 	UFUNCTION(BlueprintCallable, Category = "Settings")
 	MOVIERENDERPIPELINECORE_API void Apply(const UWorld* World);
 
@@ -1100,7 +1123,7 @@ private:
 
 	/** The modifiers that are active when this render layer is active. */
 	UPROPERTY()
-	TArray<TObjectPtr<UMovieGraphCollectionModifier>> Modifiers;
+	TArray<TObjectPtr<UMovieGraphModifierBase>> Modifiers;
 };
 
 /**
diff --git a/Engine/Plugins/MovieScene/MovieRenderPipeline/Source/MovieRenderPipelineCore/Public/Graph/Nodes/MovieGraphMaterialParameterCollectionModifier.h b/Engine/Plugins/MovieScene/MovieRenderPipeline/Source/MovieRenderPipelineCore/Public/Graph/Nodes/MovieGraphMaterialParameterCollectionModifier.h
new file mode 100644
--- /dev/null
+++ b/Engine/Plugins/MovieScene/MovieRenderPipeline/Source/MovieRenderPipelineCore/Public/Graph/Nodes/MovieGraphMaterialParameterCollectionModifier.h
@@ -0,0 +1,63 @@
+// Copyright Epic Games, Inc. All Rights Reserved.
+
+#pragma once
+
+#include "Graph/MovieGraphRenderLayerSubsystem.h"
+
+#include "MovieGraphMaterialParameterCollectionModifier.generated.h"
+
+#define UE_API MOVIERENDERPIPELINECORE_API
+
+class UMaterialParameterCollection;
+
+/**
+ * A modifier that allows for changing scalar and vector parameters within a Material Parameter Collection.
+ */
+UCLASS(MinimalAPI, BlueprintType)
+class UMovieGraphMaterialParameterCollectionModifier : public UMovieGraphModifierBase
+{
+	GENERATED_BODY()
+
+public:
+	//~ Begin UMovieGraphModifierBase Interface
+	virtual void ApplyModifier(const UWorld* World) override;
+	virtual void UndoModifier() override;
+	//~ End UMovieGraphModifierBase Interface
+
+	/** Clear the scalar and vector parameter values that have been set on this modifier. */
+	UFUNCTION(BlueprintCallable, Category = "Modifier")
+	UE_API void ClearParameterValues();
+
+public:
+	/** The Material Parameter Collection that this modifier should update. */
+	UPROPERTY(BlueprintReadWrite, EditAnywhere, Category = "Settings")
+	TSoftObjectPtr<UMaterialParameterCollection> MaterialParameterCollection;
+
+	/** The names and values of scalar parameters that this modifier will apply. */
+	UPROPERTY(BlueprintReadWrite, EditAnywhere, Category = "Settings")
+	TMap<FName, float> ScalarParameterUpdates;
+
+	/** The names and values of vector parameters that this modifier will apply. */
+	UPROPERTY(BlueprintReadWrite, EditAnywhere, Category = "Settings")
+	TMap<FName, FLinearColor> VectorParameterUpdates;
+
+private:
+	/**
+	 * Applies the given scalar and vector values to the Material Parameter Collection associated with this modifier.
+	 * Optionally, the current values can be cached for later restoration.
+	 */
+	void ApplyScalarAndVectorValues(const TMap<FName, float>& InScalarValues, const TMap<FName, FLinearColor>& InVectorValues, const bool bCacheCurrentValues = false);
+
+private:
+	/** The cached names/values of scalar parameters before the modifier was applied. */
+	TMap<FName, float> PriorScalarValues;
+	
+	/** The cached names/values of vector parameters before the modifier was applied. */
+	TMap<FName, FLinearColor> PriorVectorValues;
+
+	/** The world that the modifier is being applied in. */
+	UPROPERTY(Transient)
+	TWeakObjectPtr<const UWorld> CachedWorld = nullptr;
+};
+
+#undef UE_API
\ No newline at end of file
diff --git a/Engine/Plugins/MovieScene/MovieRenderPipeline/Source/MovieRenderPipelineCore/Public/Graph/Nodes/MovieGraphMaterialParameterCollectionModifierNode.h b/Engine/Plugins/MovieScene/MovieRenderPipeline/Source/MovieRenderPipelineCore/Public/Graph/Nodes/MovieGraphMaterialParameterCollectionModifierNode.h
new file mode 100644
--- /dev/null
+++ b/Engine/Plugins/MovieScene/MovieRenderPipeline/Source/MovieRenderPipelineCore/Public/Graph/Nodes/MovieGraphMaterialParameterCollectionModifierNode.h
@@ -0,0 +1,78 @@
+﻿// Copyright Epic Games, Inc. All Rights Reserved.
+
+#pragma once
+
+#include "Graph/MovieGraphRenderLayerSubsystem.h"
+#include "MovieGraphModifierNode.h"
+
+#include "MovieGraphMaterialParameterCollectionModifierNode.generated.h"
+
+class UMaterialParameterCollection;
+class UMovieGraphMaterialParameterCollectionModifier;
+
+#define UE_API MOVIERENDERPIPELINECORE_API
+
+/** 
+ * A node which modifies the scalar and vector parameters of a Material Parameter Collection. 
+ */
+UCLASS(MinimalAPI)
+class UMovieGraphMaterialParameterCollectionModifierNode : public UMovieGraphSettingNode, public IMovieGraphModifierNodeInterface
+{
+	GENERATED_BODY()
+
+public:
+	UE_API UMovieGraphMaterialParameterCollectionModifierNode();
+
+#if WITH_EDITOR
+	//~ Begin UMovieGraphNode interface
+	UE_API virtual FText GetNodeTitle(const bool bGetDescriptive = false) const override;
+	UE_API virtual FText GetMenuCategory() const override;
+	UE_API virtual FLinearColor GetNodeTitleColor() const override;
+	UE_API virtual FSlateIcon GetIconAndTint(FLinearColor& OutColor) const override;
+	//~ End UMovieGraphNode interface
+
+	//~ Begin UObject interface
+	UE_API virtual void PostEditChangeProperty(FPropertyChangedEvent& PropertyChangedEvent) override;
+	UE_API virtual void PostLoad() override;
+	//~ End UObject interface
+#endif
+
+	//~ Begin UMovieGraphNode interface
+	UE_API virtual TArray<FPropertyBagPropertyDesc> GetDynamicPropertyDescriptions() const override;
+	UE_API virtual bool GetDynamicPropertyValue(const FName PropertyName, FString& OutValue) override;
+	UE_API virtual TArray<FMovieGraphPropertyInfo> GetOverrideablePropertyInfo() const override;
+	//~ End UMovieGraphNode interface
+
+	//~ Begin IMovieGraphModifierNodeInterface interface
+	UE_API virtual TArray<UMovieGraphModifierBase*> GetAllModifiers() const override;
+	UE_API virtual bool SupportsCollections() const override;
+	//~ End IMovieGraphModifierNodeInterface interface
+
+	//~ Begin UMovieGraphSettingNode interface
+	UE_API virtual FString GetNodeInstanceName() const override;
+	UE_API virtual void PrepareForFlattening(const UMovieGraphSettingNode* InSourceNode) override;
+	//~ End UMovieGraphSettingNode interface
+
+public:	
+	UPROPERTY(EditAnywhere, BlueprintReadWrite, Category = Overrides, meta = (InlineEditConditionToggle))
+	uint8 bOverride_MaterialParameterCollection : 1;
+
+	/** The Material Parameter Collection that should be modified. */
+	UPROPERTY(BlueprintReadWrite, EditAnywhere, Category = "Material Parameter Collection", meta = (EditCondition = "bOverride_MaterialParameterCollection"))
+	TSoftObjectPtr<UMaterialParameterCollection> MaterialParameterCollection;
+
+private:
+	void OnObjectPropertyChanged(UObject* ObjectBeingModified, FPropertyChangedEvent& PropertyChangedEvent);
+	
+	// These should really be defined somewhere outside of MRG, but they don't appear to be.
+	static inline const FName EditConditionMetadataKey = TEXT("EditCondition");
+	static inline const FName CategoryMetadataKey = TEXT("Category");
+	static inline const FName DisplayNameMetadataKey = TEXT("DisplayName");
+	static inline const FName InlineEditConditionToggleMetadataKey = TEXT("InlineEditConditionToggle");
+	
+	/** The Material Parameter Collection modifier associated with this node. */
+	UPROPERTY()
+	TObjectPtr<UMovieGraphMaterialParameterCollectionModifier> Modifier;
+};
+
+#undef UE_API
diff --git a/Engine/Plugins/MovieScene/MovieRenderPipeline/Source/MovieRenderPipelineCore/Public/Graph/Nodes/MovieGraphModifierNode.h b/Engine/Plugins/MovieScene/MovieRenderPipeline/Source/MovieRenderPipelineCore/Public/Graph/Nodes/MovieGraphModifierNode.h
--- a/Engine/Plugins/MovieScene/MovieRenderPipeline/Source/MovieRenderPipelineCore/Public/Graph/Nodes/MovieGraphModifierNode.h
+++ b/Engine/Plugins/MovieScene/MovieRenderPipeline/Source/MovieRenderPipelineCore/Public/Graph/Nodes/MovieGraphModifierNode.h
@@ -30,12 +30,84 @@ public:
 	TArray<TObjectPtr<UMovieGraphCollectionModifier>> Modifiers;
 };
 
+UINTERFACE(MinimalAPI, NotBlueprintable)
+class UMovieGraphModifierNodeInterface : public UInterface
+{
+	GENERATED_BODY()
+};
+
+/**
+ * Provides the interface that nodes must implement if they want to provide modifier behavior.
+ *
+ * If a modifier node does not support collections, it must specify this explicitly via SupportsCollections(). If it does not support collections,
+ * all collection-related methods can remain unimplemented.
+ *
+ * Note that some of these methods may not be called by MRG yet, but they may be called in the future (eg, in future UIs). To ensure modifier
+ * functionality in the future, implement all methods that make sense for the modifier.
+ */
+class IMovieGraphModifierNodeInterface : public IInterface
+{
+	GENERATED_BODY()
+
+public:
+	/** Gets all modifiers that will be applied with this modifier node. */
+	UFUNCTION(BlueprintCallable, Category = "Modifiers")
+	UE_API virtual TArray<UMovieGraphModifierBase*> GetAllModifiers() const
+	{
+		return {};
+	}
+
+	/** Gets whether this modifier node supports/uses collections. */
+	UFUNCTION(BlueprintCallable, Category = "Modifiers")
+	UE_API virtual bool SupportsCollections() const
+	{
+		return true;
+	}
+
+	/** Gets the enable state (within this modifier) of the collection with the given name. */
+	UFUNCTION(BlueprintCallable, Category = "Modifiers")
+	UE_API virtual bool IsCollectionEnabled(const FName& InCollectionName) const
+	{
+		return true;
+	}
+
+	/**
+	 * Sets the enable state (within this modifier) of the collection with the given name. Disabled collections will not be modified by this modifier
+	 * node. Collections that are added to the modifier are enabled by default.
+	 */
+	UFUNCTION(BlueprintCallable, Category = "Modifiers")
+	UE_API virtual void SetCollectionEnabled(const FName& InCollectionName, const bool bIsCollectionEnabled)
+	{
+		
+	}
+
+	/** Gets all collections that will be affected by the modifiers used by this node. */
+	UFUNCTION(BlueprintCallable, Category = "Modifiers")
+	UE_API virtual TArray<FName> GetAllCollections() const
+	{
+		return {};
+	}
+
+	/** Adds a collection identified by the given name which will be affected by the modifiers on this node. */
+	UFUNCTION(BlueprintCallable, Category = "Modifiers")
+	UE_API virtual void AddCollection(const FName& InCollectionName)
+	{
+		
+	}
+
+	/** Removes a collection identified by the given name. Returns true if the collection was found and removed successfully, else false. */
+	UFUNCTION(BlueprintCallable, Category = "Modifiers")
+	UE_API virtual bool RemoveCollection(const FName& InCollectionName)
+	{
+		return false;
+	}
+};
+
 /** 
-* A collection node specifies an interface for doing dynamic scene queries for actors in the world. Collections work in tandem with
-* UMovieGraphModifiers to select which actors the modifiers should be run on.
-*/
+ * A modifier node which allows render properties and materials to be changed.
+ */
 UCLASS(MinimalAPI)
-class UMovieGraphModifierNode : public UMovieGraphSettingNode
+class UMovieGraphModifierNode : public UMovieGraphSettingNode, public IMovieGraphModifierNodeInterface
 {
 	GENERATED_BODY()
 
@@ -43,23 +115,38 @@ public:
 	UE_API UMovieGraphModifierNode();
 
 #if WITH_EDITOR
+	//~ Begin UMovieGraphNode interface
 	UE_API virtual FText GetNodeTitle(const bool bGetDescriptive = false) const override;
 	UE_API virtual FText GetMenuCategory() const override;
 	UE_API virtual FLinearColor GetNodeTitleColor() const override;
 	UE_API virtual FSlateIcon GetIconAndTint(FLinearColor& OutColor) const override;
+	//~ End UMovieGraphNode interface
 
 	//~ Begin UObject interface
 	UE_API virtual void PostEditChangeProperty(FPropertyChangedEvent& PropertyChangedEvent) override;
 	//~ End UObject interface
 #endif
 
-	virtual FString GetNodeInstanceName() const override { return ModifierName; }
+	//~ Begin UMovieGraphSettingNode interface
+	UE_API virtual FString GetNodeInstanceName() const override;
+	//~ End UMovieGraphSettingNode interface
+
+	//~ Begin IMovieGraphModifierNodeInterface
+	UE_API virtual TArray<UMovieGraphModifierBase*> GetAllModifiers() const override;
+	UE_API virtual void AddCollection(const FName& InCollectionName) override;
+	UE_API virtual bool RemoveCollection(const FName& InCollectionName) override;
+	UE_API virtual TArray<FName> GetAllCollections() const override;
+	UE_API virtual bool SupportsCollections() const override;
+	UE_API virtual void SetCollectionEnabled(const FName& InCollectionName, const bool bIsCollectionEnabled) override;
+	UE_API virtual bool IsCollectionEnabled(const FName& InCollectionName) const override;
+	//~ End IMovieGraphModifierNodeInterface
 
 	/** Gets the modifier of the specified type, or nullptr if one does not exist on this node. */
 	UFUNCTION(BlueprintCallable, Category = "Modifiers")
 	UE_API UMovieGraphCollectionModifier* GetModifier(TSubclassOf<UMovieGraphCollectionModifier> ModifierType) const;
 
 	/** Gets all modifiers currently added to the node. */
+	UE_DEPRECATED(5.7, "Use GetAllModifiers() instead.")
 	UFUNCTION(BlueprintCallable, Category = "Modifiers")
 	UE_API const TArray<UMovieGraphCollectionModifier*>& GetModifiers() const;
 
@@ -74,29 +161,11 @@ public:
 	UFUNCTION(BlueprintCallable, Category = "Modifiers")
 	UE_API bool RemoveModifier(TSubclassOf<UMovieGraphCollectionModifier> ModifierType);
 
-	/** Add a collection identified by the given name which will be affected by the modifiers on this node. */
-	UFUNCTION(BlueprintCallable, Category = "Modifiers")
-	UE_API void AddCollection(const FName& InCollectionName);
-
-	/** Remove a collection identified by the given name. Returns true if the collection was found and removed successfully, else false. */
-	UFUNCTION(BlueprintCallable, Category = "Modifiers")
-	UE_API bool RemoveCollection(const FName& InCollectionName);
-
 	/** Gets all collections that will be affected by the modifiers on this node. */
+	UE_DEPRECATED(5.7, "Use GetAllCollections() instead.")
 	UFUNCTION(BlueprintCallable, Category = "Modifiers")
 	UE_API const TArray<FName>& GetCollections() const;
 
-	/**
-	 * Sets the enable state (within this modifier) of the collection with the given name. Disabled collections will not be modified by this modifier
-	 * node. Collections that are added to the modifier are enabled by default.
-	 */
-	UFUNCTION(BlueprintCallable, Category = "Modifiers")
-	UE_API void SetCollectionEnabled(const FName& InCollectionName, const bool bIsCollectionEnabled);
-
-	/** Gets the enable state (within this modifier) of the collection with the given name. */
-	UFUNCTION(BlueprintCallable, Category = "Modifiers")
-	UE_API bool IsCollectionEnabled(const FName& InCollectionName) const;
-
 public:
 	UPROPERTY()
 	uint8 bOverride_ModifierName : 1 = 1;	// Always merge the modifier name, no need for the user to do this explicitly
diff --git a/Engine/Plugins/MovieScene/MovieRenderPipeline/Source/MovieRenderPipelineEditor/Private/Customizations/Graph/MovieGraphMaterialParameterCollectionModifierNodeCustomization.h b/Engine/Plugins/MovieScene/MovieRenderPipeline/Source/MovieRenderPipelineEditor/Private/Customizations/Graph/MovieGraphMaterialParameterCollectionModifierNodeCustomization.h
new file mode 100644
--- /dev/null
+++ b/Engine/Plugins/MovieScene/MovieRenderPipeline/Source/MovieRenderPipelineEditor/Private/Customizations/Graph/MovieGraphMaterialParameterCollectionModifierNodeCustomization.h
@@ -0,0 +1,101 @@
+// Copyright Epic Games, Inc. All Rights Reserved.
+
+#pragma once
+
+#include "Graph/Nodes/MovieGraphMaterialParameterCollectionModifierNode.h"
+
+#include "DetailCategoryBuilder.h"
+#include "DetailLayoutBuilder.h"
+#include "IDetailCustomization.h"
+#include "IDetailGroup.h"
+#include "InstancedPropertyBagStructureDataProvider.h"
+
+#define LOCTEXT_NAMESPACE "MoviePipelineEditor"
+
+/** Customizes how the Material Parameter Collection Modifier node appears in the details panel. */
+class FMovieGraphMaterialParameterCollectionModifierNodeCustomization : public IDetailCustomization
+{
+public:
+	static TSharedRef<IDetailCustomization> MakeInstance()
+	{
+		return MakeShared<FMovieGraphMaterialParameterCollectionModifierNodeCustomization>();
+	}
+
+protected:
+	//~ Begin IDetailCustomization interface
+	virtual void CustomizeDetails(const TSharedPtr<IDetailLayoutBuilder>& DetailBuilder) override
+	{
+		CustomizeDetails(*DetailBuilder);
+	}
+	
+	virtual void CustomizeDetails(IDetailLayoutBuilder& DetailBuilder) override
+	{
+		static const FName ScalarParamsGroupName = FName(TEXT("ScalarParams"));
+		static const FName VectorParamsGroupName = FName(TEXT("VectorParams"));
+
+		static const FText ScalarParamsGroupText = LOCTEXT("MPCNode_ScalarParameters", "Scalar Parameters");
+		static const FText VectorParamsGroupText = LOCTEXT("MPCNode_VectorParameters", "Vector Parameters");
+		
+		for (const TWeakObjectPtr<UMovieGraphMaterialParameterCollectionModifierNode>& MPCNode
+				: DetailBuilder.GetObjectsOfTypeBeingCustomized<UMovieGraphMaterialParameterCollectionModifierNode>())
+		{
+			// Regenerate the details when the MPC asset is changed.
+			const TSharedRef<IPropertyHandle> MpcAssetHandle = DetailBuilder.GetProperty(GET_MEMBER_NAME_CHECKED(UMovieGraphMaterialParameterCollectionModifierNode, MaterialParameterCollection));
+			MpcAssetHandle->SetOnPropertyValueChanged(FSimpleDelegate::CreateLambda([&DetailBuilder]()
+			{
+				DetailBuilder.ForceRefreshDetails();
+			}));
+
+			// Hide the "Properties" category that is shown by default for the DynamicProperties property that's inherited from UMovieGraphNode. In
+			// this customization, we want to display these DynamicProperties in a non-default way.
+			DetailBuilder.HideCategory("Properties");
+
+			const TSharedRef<IPropertyHandle> DynamicPropertiesHandle = DetailBuilder.GetProperty(TEXT("DynamicProperties"), UMovieGraphNode::StaticClass());
+			if (!DynamicPropertiesHandle->IsValidHandle())
+			{
+				continue;
+			}
+
+			// Get all child handles of the dynamic properties property bag (ie, the scalar/vector properties, plus their associated bOverride_* properties).
+			TArray<TSharedPtr<IPropertyHandle>> PropertyBagChildHandles =
+				DynamicPropertiesHandle->AddChildStructure(MakeShared<FInstancePropertyBagStructureDataProvider>(*MPCNode->GetMutableDynamicProperties_Unsafe()));
+
+			// For each category (MPC asset name), keep track of the "Scalar"/"Vector" groups that have been added.
+			TMap<FString, TMap<FName, IDetailGroup*>> CategoryToGroupMap;
+
+			// Add all of the scalar/vector parameters.
+			for (const TSharedPtr<IPropertyHandle>& ChildHandle : PropertyBagChildHandles)
+			{
+				const FString CategoryName = ChildHandle->GetMetaData(FName(TEXT("Category")));
+
+				// If the category is empty, skip -- this is probably a bOverride_* property.
+				if (CategoryName.IsEmpty())
+				{
+					continue;
+				}
+
+				// Add the category (MPC asset name) if it hasn't been added already.
+				IDetailCategoryBuilder& CategoryBuilder = DetailBuilder.EditCategory(FName(CategoryName));
+
+				const bool bIsScalar = (ChildHandle->GetPropertyClass() == FFloatProperty::StaticClass());
+				const FText& GroupText = bIsScalar ? ScalarParamsGroupText : VectorParamsGroupText;
+				const FName& GroupName = bIsScalar ? ScalarParamsGroupName : VectorParamsGroupName;
+
+				// Add the sub-group ("Scalar"/"Vector" parameters) if it hasn't been added already.
+				IDetailGroup*& ParamGroup = CategoryToGroupMap.FindOrAdd(CategoryName).FindOrAdd(GroupName);
+				if (!ParamGroup)
+				{
+					constexpr bool bForAdvanced = false;
+					constexpr bool bStartExpanded = true;
+					ParamGroup = &CategoryBuilder.AddGroup(GroupName, GroupText, bForAdvanced, bStartExpanded);
+				}
+
+				// Add the parameter to the details panel as a new row.
+				ParamGroup->AddPropertyRow(ChildHandle.ToSharedRef());
+			}
+		}
+	}
+	//~ End IDetailCustomization interface
+};
+
+#undef LOCTEXT_NAMESPACE
\ No newline at end of file
diff --git a/Engine/Plugins/MovieScene/MovieRenderPipeline/Source/MovieRenderPipelineEditor/Private/Customizations/Graph/MovieGraphModifiersCustomization.cpp b/Engine/Plugins/MovieScene/MovieRenderPipeline/Source/MovieRenderPipelineEditor/Private/Customizations/Graph/MovieGraphModifiersCustomization.cpp
--- a/Engine/Plugins/MovieScene/MovieRenderPipeline/Source/MovieRenderPipelineEditor/Private/Customizations/Graph/MovieGraphModifiersCustomization.cpp
+++ b/Engine/Plugins/MovieScene/MovieRenderPipeline/Source/MovieRenderPipelineEditor/Private/Customizations/Graph/MovieGraphModifiersCustomization.cpp
@@ -162,7 +162,7 @@ void FMovieGraphModifiersCustomization::CustomizeDetails(IDetailLayoutBuilder& I
 					{
 						if (ModifierNode.IsValid())
 						{
-							return !ModifierNode.Get()->GetCollections().Contains(CollectionName);
+							return !ModifierNode.Get()->GetAllCollections().Contains(CollectionName);
 						}
 
 						return false;
@@ -174,7 +174,7 @@ void FMovieGraphModifiersCustomization::CustomizeDetails(IDetailLayoutBuilder& I
 							const FScopedTransaction Transaction(LOCTEXT("AddCollectionToModifier", "Add Collection to Modifier"));
 							
 							ModifierNode->AddCollection(PickedCollectionName);
-							ListDataSource = ModifierNode->GetCollections();
+							ListDataSource = ModifierNode->GetAllCollections();
 							CollectionsList->Refresh();
 						}
 					});
@@ -210,7 +210,7 @@ void FMovieGraphModifiersCustomization::CustomizeDetails(IDetailLayoutBuilder& I
 					ModifierNode.Get()->RemoveCollection(DeletedCollectionName);
 				}
 				
-				ListDataSource = ModifierNode->GetCollections();
+				ListDataSource = ModifierNode->GetAllCollections();
 				CollectionsList->Refresh();
 			}
 		})
@@ -233,7 +233,7 @@ void FMovieGraphModifiersCustomization::CustomizeDetails(IDetailLayoutBuilder& I
 	];
 
 	// For all modifiers added to the node, add a category for each, and add each modifier's EditAnywhere properties to the category
-	for (UMovieGraphCollectionModifier* Modifier : ModifierNode->GetModifiers())
+	for (UMovieGraphModifierBase* Modifier : ModifierNode->GetAllModifiers())
 	{
 		if (!Modifier)
 		{
@@ -302,7 +302,7 @@ void FMovieGraphModifiersCustomization::RefreshListDataSource()
 	if (const TStrongObjectPtr<UMovieGraphModifierNode> ModifierNodePin = ModifierNode.Pin())
 	{
 		// Update the data source
-		ListDataSource = ModifierNodePin->GetCollections();
+		ListDataSource = ModifierNodePin->GetAllCollections();
 	}
 }
 
diff --git a/Engine/Plugins/MovieScene/MovieRenderPipeline/Source/MovieRenderPipelineEditor/Private/Graph/MovieGraphAssetToolkit.cpp b/Engine/Plugins/MovieScene/MovieRenderPipeline/Source/MovieRenderPipelineEditor/Private/Graph/MovieGraphAssetToolkit.cpp
--- a/Engine/Plugins/MovieScene/MovieRenderPipeline/Source/MovieRenderPipelineEditor/Private/Graph/MovieGraphAssetToolkit.cpp
+++ b/Engine/Plugins/MovieScene/MovieRenderPipeline/Source/MovieRenderPipelineEditor/Private/Graph/MovieGraphAssetToolkit.cpp
@@ -6,6 +6,7 @@
 #include "Customizations/Graph/MovieGraphCollectionsCustomization.h"
 #include "Customizations/Graph/MovieGraphMultiLayerEXRNodeCustomization.h"
 #include "Customizations/Graph/MovieGraphFormatTokenCustomization.h"
+#include "Customizations/Graph/MovieGraphMaterialParameterCollectionModifierNodeCustomization.h"
 #include "Customizations/Graph/MovieGraphMemberCustomization.h"
 #include "Customizations/Graph/MovieGraphMetadataAttributeCustomization.h"
 #include "Customizations/Graph/MovieGraphModifiersCustomization.h"
@@ -23,6 +24,7 @@
 #include "Graph/Nodes/MovieGraphCommandLineEncoderNode.h"
 #include "Graph/Nodes/MovieGraphDebugNode.h"
 #include "Graph/Nodes/MovieGraphFileOutputNode.h"
+#include "Graph/Nodes/MovieGraphMaterialParameterCollectionModifierNode.h"
 #include "Graph/Nodes/MovieGraphModifierNode.h"
 #include "Graph/Nodes/MovieGraphPathTracerPassNode.h"
 #include "Graph/Nodes/MovieGraphSelectNode.h"
@@ -403,6 +405,10 @@ TSharedRef<SDockTab> FMovieGraphAssetToolkit::SpawnTab_RenderGraphDetails(const
 		UMovieGraphApplyCVarPresetNode::StaticClass(),
 		FOnGetDetailCustomizationInstance::CreateStatic(&FMovieGraphApplyCVarPresetNodeCustomization::MakeInstance));
 
+	SelectedGraphObjectsDetailsWidget->RegisterInstancedCustomPropertyLayout(
+		UMovieGraphMaterialParameterCollectionModifierNode::StaticClass(),
+		FOnGetDetailCustomizationInstance::CreateStatic(&FMovieGraphMaterialParameterCollectionModifierNodeCustomization::MakeInstance));
+
 	SelectedGraphObjectsDetailsWidget->RegisterInstancedCustomPropertyTypeLayout(
 		UMovieGraphShowFlags::StaticClass()->GetFName(),
 		FOnGetPropertyTypeCustomizationInstance::CreateStatic(&FMovieGraphShowFlagsCustomization::MakeInstance));
diff --git a/Engine/Source/Runtime/Engine/Public/Materials/MaterialParameterCollection.h b/Engine/Source/Runtime/Engine/Public/Materials/MaterialParameterCollection.h
--- a/Engine/Source/Runtime/Engine/Public/Materials/MaterialParameterCollection.h
+++ b/Engine/Source/Runtime/Engine/Public/Materials/MaterialParameterCollection.h
@@ -154,7 +154,7 @@ public:
 	//~ End UObject Interface.
 
 	/** Finds a parameter name given an Id, returns NAME_None if the parameter was not found. */
-	FName GetParameterName(const FGuid& Id) const;
+	ENGINE_API FName GetParameterName(const FGuid& Id) const;
 
 	/** Finds a parameter id given a name, returns the default guid if the parameter was not found. */
 	ENGINE_API FGuid GetParameterId(FName ParameterName) const;
