MoviePipeline: Gave graph modifiers a new GetModifierName() method so they can use this (rather than the DisplayName class metadata) for display in the UI.

diff --git a/Engine/Plugins/MovieScene/MovieRenderPipeline/Source/MovieRenderPipelineCore/Private/Graph/MovieGraphRenderLayerSubsystem.cpp b/Engine/Plugins/MovieScene/MovieRenderPipeline/Source/MovieRenderPipelineCore/Private/Graph/MovieGraphRenderLayerSubsystem.cpp
--- a/Engine/Plugins/MovieScene/MovieRenderPipeline/Source/MovieRenderPipelineCore/Private/Graph/MovieGraphRenderLayerSubsystem.cpp
+++ b/Engine/Plugins/MovieScene/MovieRenderPipeline/Source/MovieRenderPipelineCore/Private/Graph/MovieGraphRenderLayerSubsystem.cpp
@@ -172,6 +172,12 @@ void UMovieGraphMaterialModifier::UndoModifier()
 	ModifiedComponents.Empty();
 }
 
+FText UMovieGraphMaterialModifier::GetModifierName()
+{
+	static const FText ModifierName = LOCTEXT("MaterialModifierName", "Material");
+	return ModifierName;
+}
+
 UMovieGraphRenderPropertyModifier::UMovieGraphRenderPropertyModifier()
 	: bIsHidden(false)
 	, bCastsShadows(true)
@@ -316,6 +322,12 @@ void UMovieGraphRenderPropertyModifier::UndoModifier()
 	CachedComponentState.Empty();
 }
 
+FText UMovieGraphRenderPropertyModifier::GetModifierName()
+{
+	static const FText ModifierName = LOCTEXT("RenderPropertyModifierName", "Visibility");
+	return ModifierName;
+}
+
 void UMovieGraphRenderPropertyModifier::SetActorAndComponentState(const TArray<FActorState>& InActorState, const TArray<FComponentState>& InComponentState, const bool bUseStateFromNode)
 {
 	// In most cases, if the hidden state is being modified, the hidden state should be set. However, there is an exception for volumetrics.
@@ -509,6 +521,13 @@ TArray<UMovieGraphCollection*> UMovieGraphCollectionModifier::GetCollections() c
 	return Collections;
 }
 
+FText UMovieGraphCollectionModifier::GetModifierName()
+{
+	// Derived classes should override this method
+	static const FText ModifierName = LOCTEXT("CollectionModifierName", "Unknown Collection Modifier");
+	return ModifierName;
+}
+
 UMovieGraphConditionGroupQueryBase::UMovieGraphConditionGroupQueryBase()
 	: OpType(EMovieGraphConditionGroupQueryOpType::Add)
 	, bIsEnabled(true)
diff --git a/Engine/Plugins/MovieScene/MovieRenderPipeline/Source/MovieRenderPipelineCore/Private/Graph/Nodes/MovieGraphMaterialParameterCollectionModifier.cpp b/Engine/Plugins/MovieScene/MovieRenderPipeline/Source/MovieRenderPipelineCore/Private/Graph/Nodes/MovieGraphMaterialParameterCollectionModifier.cpp
--- a/Engine/Plugins/MovieScene/MovieRenderPipeline/Source/MovieRenderPipelineCore/Private/Graph/Nodes/MovieGraphMaterialParameterCollectionModifier.cpp
+++ b/Engine/Plugins/MovieScene/MovieRenderPipeline/Source/MovieRenderPipelineCore/Private/Graph/Nodes/MovieGraphMaterialParameterCollectionModifier.cpp
@@ -24,6 +24,12 @@ void UMovieGraphMaterialParameterCollectionModifier::UndoModifier()
 	PriorVectorValues.Reset();
 }
 
+FText UMovieGraphMaterialParameterCollectionModifier::GetModifierName()
+{
+	static const FText ModifierName = NSLOCTEXT("MovieGraphNode", "MaterialParameterCollectionModifierName", "Material Parameter Collection");
+	return ModifierName;
+}
+
 void UMovieGraphMaterialParameterCollectionModifier::ClearParameterValues()
 {
 	ScalarParameterUpdates.Reset();
diff --git a/Engine/Plugins/MovieScene/MovieRenderPipeline/Source/MovieRenderPipelineCore/Public/Graph/MovieGraphRenderLayerSubsystem.h b/Engine/Plugins/MovieScene/MovieRenderPipeline/Source/MovieRenderPipelineCore/Public/Graph/MovieGraphRenderLayerSubsystem.h
--- a/Engine/Plugins/MovieScene/MovieRenderPipeline/Source/MovieRenderPipelineCore/Public/Graph/MovieGraphRenderLayerSubsystem.h
+++ b/Engine/Plugins/MovieScene/MovieRenderPipeline/Source/MovieRenderPipelineCore/Public/Graph/MovieGraphRenderLayerSubsystem.h
@@ -881,8 +881,17 @@ class UMovieGraphModifierBase : public UObject
 	GENERATED_BODY()
 
 public:
+	/** Applies this modifier in the provided world. Called once per layer. */
+	UFUNCTION(BlueprintCallable, Category = "Modifier")
 	MOVIERENDERPIPELINECORE_API virtual void ApplyModifier(const UWorld* World) PURE_VIRTUAL(UMovieGraphModifierBase::ApplyModifier, );
+
+	/** Undoes the effects of this modifier. The modifier is responsible for tracking changes that it makes. Called once per layer. */
+	UFUNCTION(BlueprintCallable, Category = "Modifier")
 	MOVIERENDERPIPELINECORE_API virtual void UndoModifier() PURE_VIRTUAL(UMovieGraphModifierBase::UndoModifier, );
+
+	/** Gets the name of this modifier. Typically used for UI display purposes. */
+	UFUNCTION(BlueprintCallable, Category = "Modifier")
+	MOVIERENDERPIPELINECORE_API virtual FText GetModifierName() PURE_VIRTUAL(UMovieGraphModifierBase::GetModifierName, return FText::GetEmpty(););
 };
 
 /**
@@ -906,6 +915,10 @@ public:
 	UFUNCTION(BlueprintCallable, Category = "Settings")
 	MOVIERENDERPIPELINECORE_API TArray<UMovieGraphCollection*> GetCollections() const;
 
+	//~ Begin UMovieGraphModifierBase Interface
+	MOVIERENDERPIPELINECORE_API virtual FText GetModifierName() override;
+	//~ End UMovieGraphModifierBase Interface
+
 protected:
 	/** The collections which this modifier will operate on. */
 	UPROPERTY()
@@ -915,7 +928,7 @@ protected:
 /**
  * Modifies actor materials.
  */
-UCLASS(MinimalAPI, BlueprintType, DisplayName="Material")
+UCLASS(MinimalAPI, BlueprintType)
 class UMovieGraphMaterialModifier : public UMovieGraphCollectionModifier
 {
 	GENERATED_BODY()
@@ -924,11 +937,11 @@ public:
 	UFUNCTION(BlueprintCallable, Category = "Settings")
 	void SetMaterial(TSoftObjectPtr<UMaterialInterface> InMaterial) { Material = InMaterial; }
 
-	UFUNCTION(BlueprintCallable, Category = "Settings")
+	//~ Begin UMovieGraphModifierBase Interface
 	MOVIERENDERPIPELINECORE_API virtual void ApplyModifier(const UWorld* World) override;
-
-	UFUNCTION(BlueprintCallable, Category = "Settings")
 	MOVIERENDERPIPELINECORE_API virtual void UndoModifier() override;
+	MOVIERENDERPIPELINECORE_API virtual FText GetModifierName() override;
+	//~ End UMovieGraphModifierBase Interface
 
 private:
 	typedef TTuple<int32, TSoftObjectPtr<UMaterialInterface>> FMaterialSlotAssignment;
@@ -948,7 +961,7 @@ public:
 /**
  * Modifies actor visibility.
  */
-UCLASS(MinimalAPI, BlueprintType, DisplayName="Visibility")
+UCLASS(MinimalAPI, BlueprintType)
 class UMovieGraphRenderPropertyModifier : public UMovieGraphCollectionModifier
 {
 	GENERATED_BODY()
@@ -969,11 +982,11 @@ public:
 	UFUNCTION(BlueprintCallable, Category = "Settings")
 	bool IsHidden() const { return bIsHidden; }
 
-	UFUNCTION(BlueprintCallable, Category = "Settings")
+	//~ Begin UMovieGraphModifierBase Interface
 	MOVIERENDERPIPELINECORE_API virtual void ApplyModifier(const UWorld* World) override;
-
-	UFUNCTION(BlueprintCallable, Category = "Settings")
 	MOVIERENDERPIPELINECORE_API virtual void UndoModifier() override;
+	MOVIERENDERPIPELINECORE_API virtual FText GetModifierName() override;
+	//~ End UMovieGraphModifierBase Interface
 
 private:
 	/** Various visibility properties for an actor. */
diff --git a/Engine/Plugins/MovieScene/MovieRenderPipeline/Source/MovieRenderPipelineCore/Public/Graph/Nodes/MovieGraphMaterialParameterCollectionModifier.h b/Engine/Plugins/MovieScene/MovieRenderPipeline/Source/MovieRenderPipelineCore/Public/Graph/Nodes/MovieGraphMaterialParameterCollectionModifier.h
--- a/Engine/Plugins/MovieScene/MovieRenderPipeline/Source/MovieRenderPipelineCore/Public/Graph/Nodes/MovieGraphMaterialParameterCollectionModifier.h
+++ b/Engine/Plugins/MovieScene/MovieRenderPipeline/Source/MovieRenderPipelineCore/Public/Graph/Nodes/MovieGraphMaterialParameterCollectionModifier.h
@@ -20,8 +20,9 @@ class UMovieGraphMaterialParameterCollectionModifier : public UMovieGraphModifie
 
 public:
 	//~ Begin UMovieGraphModifierBase Interface
-	virtual void ApplyModifier(const UWorld* World) override;
-	virtual void UndoModifier() override;
+	UE_API virtual void ApplyModifier(const UWorld* World) override;
+	UE_API virtual void UndoModifier() override;
+	UE_API virtual FText GetModifierName() override;
 	//~ End UMovieGraphModifierBase Interface
 
 	/** Clear the scalar and vector parameter values that have been set on this modifier. */
diff --git a/Engine/Plugins/MovieScene/MovieRenderPipeline/Source/MovieRenderPipelineEditor/Private/Customizations/Graph/MovieGraphModifiersCustomization.cpp b/Engine/Plugins/MovieScene/MovieRenderPipeline/Source/MovieRenderPipelineEditor/Private/Customizations/Graph/MovieGraphModifiersCustomization.cpp
--- a/Engine/Plugins/MovieScene/MovieRenderPipeline/Source/MovieRenderPipelineEditor/Private/Customizations/Graph/MovieGraphModifiersCustomization.cpp
+++ b/Engine/Plugins/MovieScene/MovieRenderPipeline/Source/MovieRenderPipelineEditor/Private/Customizations/Graph/MovieGraphModifiersCustomization.cpp
@@ -241,7 +241,7 @@ void FMovieGraphModifiersCustomization::CustomizeDetails(IDetailLayoutBuilder& I
 		}
 		
 		const UClass* ModifierClass = Modifier->GetClass();
-		const FText DisplayName = ModifierClass->GetDisplayNameText();
+		const FText DisplayName = Modifier->GetModifierName();
 
 		// Add category as "Uncommon" to display after the general modifier properties
 		IDetailCategoryBuilder& Category = InDetailBuilder.EditCategory(FName(DisplayName.ToString()), DisplayName, ECategoryPriority::Uncommon);
