MoviePipeline: Fixed an issue in the graph that could cause the Visibility modifier to incorrectly undo its state in a rare scenario.

If a Visibility modifier is modifying collections which overlap in the actors/components that they match, the new state could be applied incorrectly. This is what would happen before this change: evaluate collection 1 -> cache matching actor/component state -> apply modification, then repeat the process for collection 2. In this change, the modification step is deferred until ALL matching actor/component state has been cached across the all of modifier's collections, THEN the modifier is applied.

The new order of operations looks like: evaluate collection 1 -> cache matching actor/component state, evaluate collection 2 -> cache matching actor/component state, apply modification. This way, actors that are matched multiple times have their state properly cached before ANY modifications.

diff --git a/Engine/Plugins/MovieScene/MovieRenderPipeline/Source/MovieRenderPipelineCore/Private/Graph/MovieGraphRenderLayerSubsystem.cpp b/Engine/Plugins/MovieScene/MovieRenderPipeline/Source/MovieRenderPipelineCore/Private/Graph/MovieGraphRenderLayerSubsystem.cpp
--- a/Engine/Plugins/MovieScene/MovieRenderPipeline/Source/MovieRenderPipelineCore/Private/Graph/MovieGraphRenderLayerSubsystem.cpp
+++ b/Engine/Plugins/MovieScene/MovieRenderPipeline/Source/MovieRenderPipelineCore/Private/Graph/MovieGraphRenderLayerSubsystem.cpp
@@ -274,6 +274,7 @@ void UMovieGraphRenderPropertyModifier::ApplyModifier(const UWorld* World)
 
 	ValidateProjectSettings();
 	
+	// Cache matching actor/component state for all collections before applying the modifier.
 	for (const UMovieGraphCollection* Collection : Collections)
 	{
 		if (!Collection)
@@ -306,11 +307,11 @@ void UMovieGraphRenderPropertyModifier::ApplyModifier(const UWorld* World)
 				CacheComponentState(SceneComponent);
 			}
 		}
-
-		// Actor/component state has been cached. Apply the node's state (settings) to all affected (ie, cached) actors/components.
-		constexpr bool bUseStateFromNode = true;
-		SetActorAndComponentState({}, {}, bUseStateFromNode);
 	}
+
+	// Actor/component state has been cached. Apply the node's state (settings) to all affected (ie, cached) actors/components.
+	constexpr bool bUseStateFromNode = true;
+	SetActorAndComponentState({}, {}, bUseStateFromNode);
 }
 
 void UMovieGraphRenderPropertyModifier::UndoModifier()
