MoviePipeline: Added some additional telemetry to the graph to track the types of image nodes used.

diff --git a/Engine/Plugins/MovieScene/MovieRenderPipeline/Source/MovieRenderPipelineCore/Private/MoviePipelineTelemetry.cpp b/Engine/Plugins/MovieScene/MovieRenderPipeline/Source/MovieRenderPipelineCore/Private/MoviePipelineTelemetry.cpp
--- a/Engine/Plugins/MovieScene/MovieRenderPipeline/Source/MovieRenderPipelineCore/Private/MoviePipelineTelemetry.cpp
+++ b/Engine/Plugins/MovieScene/MovieRenderPipeline/Source/MovieRenderPipelineCore/Private/MoviePipelineTelemetry.cpp
@@ -115,6 +115,11 @@ void FMoviePipelineTelemetry::SendBeginShotRenderTelemetry(UMoviePipelineExecuto
 	EventAttributes.Add(FAnalyticsEventAttribute(TEXT("UsesAvid"), ShotRenderTelemetry.bUsesAvid));
 	EventAttributes.Add(FAnalyticsEventAttribute(TEXT("UsesProRes"), ShotRenderTelemetry.bUsesProRes));
 	EventAttributes.Add(FAnalyticsEventAttribute(TEXT("UsesMP4"), ShotRenderTelemetry.bUsesMP4));
+	EventAttributes.Add(FAnalyticsEventAttribute(TEXT("UsesJPG"), ShotRenderTelemetry.bUsesJPG));
+	EventAttributes.Add(FAnalyticsEventAttribute(TEXT("UsesPNG"), ShotRenderTelemetry.bUsesPNG));
+	EventAttributes.Add(FAnalyticsEventAttribute(TEXT("UsesBMP"), ShotRenderTelemetry.bUsesBMP));
+	EventAttributes.Add(FAnalyticsEventAttribute(TEXT("UsesEXR"), ShotRenderTelemetry.bUsesEXR));
+	EventAttributes.Add(FAnalyticsEventAttribute(TEXT("UsesMultiEXR"), ShotRenderTelemetry.bUsesMultiEXR));
 	EventAttributes.Add(FAnalyticsEventAttribute(TEXT("ResolutionX"), ShotRenderTelemetry.ResolutionX));
 	EventAttributes.Add(FAnalyticsEventAttribute(TEXT("ResolutionY"), ShotRenderTelemetry.ResolutionY));
 	EventAttributes.Add(FAnalyticsEventAttribute(TEXT("HandleFrameCount"), ShotRenderTelemetry.HandleFrameCount));
diff --git a/Engine/Plugins/MovieScene/MovieRenderPipeline/Source/MovieRenderPipelineCore/Public/MoviePipelineTelemetry.h b/Engine/Plugins/MovieScene/MovieRenderPipeline/Source/MovieRenderPipelineCore/Public/MoviePipelineTelemetry.h
--- a/Engine/Plugins/MovieScene/MovieRenderPipeline/Source/MovieRenderPipelineCore/Public/MoviePipelineTelemetry.h
+++ b/Engine/Plugins/MovieScene/MovieRenderPipeline/Source/MovieRenderPipelineCore/Public/MoviePipelineTelemetry.h
@@ -31,6 +31,11 @@ struct FMoviePipelineShotRenderTelemetry
 	bool bUsesAvid = false;
 	bool bUsesProRes = false;
 	bool bUsesMP4 = false;
+	bool bUsesJPG = false;
+	bool bUsesPNG = false;
+	bool bUsesBMP = false;
+	bool bUsesEXR = false;
+	bool bUsesMultiEXR = false;
 	int32 ResolutionX = 0;
 	int32 ResolutionY = 0;
 	int32 HandleFrameCount = 0;
@@ -42,7 +47,7 @@ struct FMoviePipelineShotRenderTelemetry
 	bool bUsesPageToSystemMemory = false;
 
 	// Note: If adding an entry here, make sure to also update FMoviePipelineTelemetry::SendBeginShotRenderTelemetry()
-	// Also remember to track the telemetry in both the graph and legacy.
+	// Also remember to track the telemetry in both the graph and legacy, and perform schematization.
 };
 
 /** Responsible for sending out telemetry for Movie Render Queue, Movie Render Graph, and Quick Render. */
diff --git a/Engine/Plugins/MovieScene/MovieRenderPipeline/Source/MovieRenderPipelineRenderPasses/Private/MovieGraphImageSequenceOutputNode.cpp b/Engine/Plugins/MovieScene/MovieRenderPipeline/Source/MovieRenderPipelineRenderPasses/Private/MovieGraphImageSequenceOutputNode.cpp
--- a/Engine/Plugins/MovieScene/MovieRenderPipeline/Source/MovieRenderPipelineRenderPasses/Private/MovieGraphImageSequenceOutputNode.cpp
+++ b/Engine/Plugins/MovieScene/MovieRenderPipeline/Source/MovieRenderPipelineRenderPasses/Private/MovieGraphImageSequenceOutputNode.cpp
@@ -12,6 +12,7 @@
 #include "Graph/MovieGraphFilenameResolveParams.h"
 #include "Graph/MovieGraphBlueprintLibrary.h"
 #include "Graph/MovieRenderGraphEditorSettings.h"
+#include "MoviePipelineTelemetry.h"
 #include "MoviePipelineUtils.h"
 #include "MoviePipelineImageSequenceOutput.h" // for FAsyncImageQuantization
 #include "MoviePipelineEXROutput.h"
@@ -580,6 +581,11 @@ void UMovieGraphImageSequenceOutputNode_EXR::InjectNodesPostEvaluation(const FNa
 	INJECT_BURN_IN_NODE(TEXT("EXR"))
 }
 
+void UMovieGraphImageSequenceOutputNode_EXR::UpdateTelemetry(FMoviePipelineShotRenderTelemetry* InTelemetry) const
+{
+	InTelemetry->bUsesEXR = true;
+}
+
 #if WITH_UNREALEXR
 void UMovieGraphImageSequenceOutputNode_MultiLayerEXR::OnReceiveImageDataImpl(UMovieGraphPipeline* InPipeline, UE::MovieGraph::FMovieGraphOutputMergerFrame* InRawFrameData, const TSet<FMovieGraphRenderDataIdentifier>& InMask)
 {
@@ -776,6 +782,11 @@ void UMovieGraphImageSequenceOutputNode_MultiLayerEXR::InjectNodesPostEvaluation
 	INJECT_BURN_IN_MULTILAYER_NODE(TEXT("MULTI_EXR"))
 }
 
+void UMovieGraphImageSequenceOutputNode_MultiLayerEXR::UpdateTelemetry(FMoviePipelineShotRenderTelemetry* InTelemetry) const
+{
+	InTelemetry->bUsesMultiEXR = true;
+}
+
 void UMovieGraphImageSequenceOutputNode_MultiLayerEXR::GetFilenameToEXROutputConfigMappings(
 	const UMovieGraphImageSequenceOutputNode_MultiLayerEXR* InParentNode, UMovieGraphPipeline* InPipeline,
 	UE::MovieGraph::FMovieGraphOutputMergerFrame* InRawFrameData, TMap<FString, FEXROutputConfigForFilename>& OutFilenameToOutputConfigs) const
@@ -895,12 +906,27 @@ void UMovieGraphImageSequenceOutputNode_BMP::InjectNodesPostEvaluation(const FNa
 	INJECT_BURN_IN_NODE(TEXT("BMP"))
 }
 
+void UMovieGraphImageSequenceOutputNode_BMP::UpdateTelemetry(FMoviePipelineShotRenderTelemetry* InTelemetry) const
+{
+	InTelemetry->bUsesBMP = true;
+}
+
 void UMovieGraphImageSequenceOutputNode_JPG::InjectNodesPostEvaluation(const FName& InBranchName, UMovieGraphEvaluatedConfig* InEvaluatedConfig, TArray<UMovieGraphSettingNode*>& OutInjectedNodes)
 {
 	INJECT_BURN_IN_NODE(TEXT("JPG"))
 }
 
+void UMovieGraphImageSequenceOutputNode_JPG::UpdateTelemetry(FMoviePipelineShotRenderTelemetry* InTelemetry) const
+{
+	InTelemetry->bUsesJPG = true;
+}
+
 void UMovieGraphImageSequenceOutputNode_PNG::InjectNodesPostEvaluation(const FName& InBranchName, UMovieGraphEvaluatedConfig* InEvaluatedConfig, TArray<UMovieGraphSettingNode*>& OutInjectedNodes)
 {
 	INJECT_BURN_IN_NODE(TEXT("PNG"))
 }
+
+void UMovieGraphImageSequenceOutputNode_PNG::UpdateTelemetry(FMoviePipelineShotRenderTelemetry* InTelemetry) const
+{
+	InTelemetry->bUsesPNG = true;
+}
diff --git a/Engine/Plugins/MovieScene/MovieRenderPipeline/Source/MovieRenderPipelineRenderPasses/Private/MoviePipelineEXROutput.cpp b/Engine/Plugins/MovieScene/MovieRenderPipeline/Source/MovieRenderPipelineRenderPasses/Private/MoviePipelineEXROutput.cpp
--- a/Engine/Plugins/MovieScene/MovieRenderPipeline/Source/MovieRenderPipelineRenderPasses/Private/MoviePipelineEXROutput.cpp
+++ b/Engine/Plugins/MovieScene/MovieRenderPipeline/Source/MovieRenderPipelineRenderPasses/Private/MoviePipelineEXROutput.cpp
@@ -15,6 +15,7 @@
 #include "MoviePipeline.h"
 #include "MoviePipelineImageQuantization.h"
 #include "MoviePipelinePrimaryConfig.h"
+#include "MoviePipelineTelemetry.h"
 #include "IOpenExrRTTIModule.h"
 #include "Modules/ModuleManager.h"
 #include "MoviePipelineUtils.h"
@@ -844,4 +845,16 @@ void UMoviePipelineImageSequenceOutput_EXR::OnReceiveImageDataImpl(FMoviePipelin
 #endif
 }
 
-#endif // WITH_UNREALEXR
\ No newline at end of file
+#endif // WITH_UNREALEXR
+
+void UMoviePipelineImageSequenceOutput_EXR::UpdateTelemetry(FMoviePipelineShotRenderTelemetry* InTelemetry) const
+{
+	if (bMultilayer)
+	{
+		InTelemetry->bUsesMultiEXR = true;
+	}
+	else
+	{
+		InTelemetry->bUsesEXR = true;
+	}
+}
\ No newline at end of file
diff --git a/Engine/Plugins/MovieScene/MovieRenderPipeline/Source/MovieRenderPipelineRenderPasses/Private/MoviePipelineImageSequenceOutput.cpp b/Engine/Plugins/MovieScene/MovieRenderPipeline/Source/MovieRenderPipelineRenderPasses/Private/MoviePipelineImageSequenceOutput.cpp
--- a/Engine/Plugins/MovieScene/MovieRenderPipeline/Source/MovieRenderPipelineRenderPasses/Private/MoviePipelineImageSequenceOutput.cpp
+++ b/Engine/Plugins/MovieScene/MovieRenderPipeline/Source/MovieRenderPipelineRenderPasses/Private/MoviePipelineImageSequenceOutput.cpp
@@ -18,6 +18,7 @@
 #include "Misc/StringFormatArg.h"
 #include "MoviePipelineOutputBase.h"
 #include "MoviePipelineImageQuantization.h"
+#include "MoviePipelineTelemetry.h"
 #include "MoviePipelineWidgetRenderSetting.h"
 #include "MoviePipelineUtils.h"
 #include "HAL/PlatformTime.h"
@@ -334,4 +335,17 @@ void UMoviePipelineImageSequenceOutputBase::GetFormatArguments(FMoviePipelineFor
 	InOutFormatArgs.FilenameArguments.Add(TEXT("render_pass"), TEXT("RenderPassName"));
 }
 
+void UMoviePipelineImageSequenceOutput_BMP::UpdateTelemetry(FMoviePipelineShotRenderTelemetry* InTelemetry) const
+{
+	InTelemetry->bUsesBMP = true;
+}
 
+void UMoviePipelineImageSequenceOutput_PNG::UpdateTelemetry(FMoviePipelineShotRenderTelemetry* InTelemetry) const
+{
+	InTelemetry->bUsesPNG = true;
+}
+
+void UMoviePipelineImageSequenceOutput_JPG::UpdateTelemetry(FMoviePipelineShotRenderTelemetry* InTelemetry) const
+{
+	InTelemetry->bUsesJPG = true;
+}
diff --git a/Engine/Plugins/MovieScene/MovieRenderPipeline/Source/MovieRenderPipelineRenderPasses/Public/MovieGraphImageSequenceOutputNode.h b/Engine/Plugins/MovieScene/MovieRenderPipeline/Source/MovieRenderPipelineRenderPasses/Public/MovieGraphImageSequenceOutputNode.h
--- a/Engine/Plugins/MovieScene/MovieRenderPipeline/Source/MovieRenderPipelineRenderPasses/Public/MovieGraphImageSequenceOutputNode.h
+++ b/Engine/Plugins/MovieScene/MovieRenderPipeline/Source/MovieRenderPipelineRenderPasses/Public/MovieGraphImageSequenceOutputNode.h
@@ -146,6 +146,10 @@ public:
 	virtual void InjectNodesPostEvaluation(const FName& InBranchName, UMovieGraphEvaluatedConfig* InEvaluatedConfig, TArray<UMovieGraphSettingNode*>& OutInjectedNodes) override;
 	//~ End IMovieGraphEvaluationNodeInjector interface
 	
+	//~ Begin UMovieGraphSettingNode interface
+	virtual void UpdateTelemetry(FMoviePipelineShotRenderTelemetry* InTelemetry) const override;
+	//~ End UMovieGraphSettingNode interface
+	
 public:
 	UPROPERTY(EditAnywhere, BlueprintReadWrite, Category = Overrides, meta = (InlineEditConditionToggle))
 	uint8 bOverride_Compression : 1;
@@ -290,6 +294,10 @@ public:
 	virtual void InjectNodesPostEvaluation(const FName& InBranchName, UMovieGraphEvaluatedConfig* InEvaluatedConfig, TArray<UMovieGraphSettingNode*>& OutInjectedNodes) override;
 	//~ End IMovieGraphEvaluationNodeInjector interface
 
+	//~ Begin UMovieGraphSettingNode interface
+	virtual void UpdateTelemetry(FMoviePipelineShotRenderTelemetry* InTelemetry) const override;
+	//~ End UMovieGraphSettingNode interface
+
 private:
 	/** Stores the EXR output config to use for a specific file name */
 	struct FEXROutputConfigForFilename
@@ -395,6 +403,10 @@ public:
 	virtual void InjectNodesPostEvaluation(const FName& InBranchName, UMovieGraphEvaluatedConfig* InEvaluatedConfig, TArray<UMovieGraphSettingNode*>& OutInjectedNodes) override;
 	//~ End IMovieGraphEvaluationNodeInjector interface
 
+	//~ Begin UMovieGraphSettingNode interface
+	virtual void UpdateTelemetry(FMoviePipelineShotRenderTelemetry* InTelemetry) const override;
+	//~ End UMovieGraphSettingNode interface
+
 public:
 	UPROPERTY(EditAnywhere, BlueprintReadWrite, Category = Overrides, meta = (InlineEditConditionToggle))
 	uint8 bOverride_bEnableBurnIn : 1;
@@ -477,6 +489,10 @@ public:
 	virtual void InjectNodesPostEvaluation(const FName& InBranchName, UMovieGraphEvaluatedConfig* InEvaluatedConfig, TArray<UMovieGraphSettingNode*>& OutInjectedNodes) override;
 	//~ End IMovieGraphEvaluationNodeInjector interface
 
+	//~ Begin UMovieGraphSettingNode interface
+	virtual void UpdateTelemetry(FMoviePipelineShotRenderTelemetry* InTelemetry) const override;
+	//~ End UMovieGraphSettingNode interface
+
 public:
 	UPROPERTY(EditAnywhere, BlueprintReadWrite, Category = Overrides, meta = (InlineEditConditionToggle))
 	uint8 bOverride_bEnableBurnIn : 1;
@@ -560,6 +576,10 @@ public:
 	virtual void InjectNodesPostEvaluation(const FName& InBranchName, UMovieGraphEvaluatedConfig* InEvaluatedConfig, TArray<UMovieGraphSettingNode*>& OutInjectedNodes) override;
 	//~ End IMovieGraphEvaluationNodeInjector interface
 
+	//~ Begin UMovieGraphSettingNode interface
+	virtual void UpdateTelemetry(FMoviePipelineShotRenderTelemetry* InTelemetry) const override;
+	//~ End UMovieGraphSettingNode interface
+
 public:
 	UPROPERTY(EditAnywhere, BlueprintReadWrite, Category = Overrides, meta = (InlineEditConditionToggle))
 	uint8 bOverride_bEnableBurnIn : 1;
diff --git a/Engine/Plugins/MovieScene/MovieRenderPipeline/Source/MovieRenderPipelineRenderPasses/Public/MoviePipelineEXROutput.h b/Engine/Plugins/MovieScene/MovieRenderPipeline/Source/MovieRenderPipelineRenderPasses/Public/MoviePipelineEXROutput.h
--- a/Engine/Plugins/MovieScene/MovieRenderPipeline/Source/MovieRenderPipelineRenderPasses/Public/MoviePipelineEXROutput.h
+++ b/Engine/Plugins/MovieScene/MovieRenderPipeline/Source/MovieRenderPipelineRenderPasses/Public/MoviePipelineEXROutput.h
@@ -217,6 +217,8 @@ public:
 	virtual bool IsAlphaAllowed() const override { return true; }
 	virtual bool ShouldCropOverscanImpl() const override { return false; }
 
+	virtual void UpdateTelemetry(FMoviePipelineShotRenderTelemetry* InTelemetry) const override;
+
 public:
 	/**
 	* Which compression method should the resulting EXR file be compressed with
diff --git a/Engine/Plugins/MovieScene/MovieRenderPipeline/Source/MovieRenderPipelineRenderPasses/Public/MoviePipelineImageSequenceOutput.h b/Engine/Plugins/MovieScene/MovieRenderPipeline/Source/MovieRenderPipelineRenderPasses/Public/MoviePipelineImageSequenceOutput.h
--- a/Engine/Plugins/MovieScene/MovieRenderPipeline/Source/MovieRenderPipelineRenderPasses/Public/MoviePipelineImageSequenceOutput.h
+++ b/Engine/Plugins/MovieScene/MovieRenderPipeline/Source/MovieRenderPipelineRenderPasses/Public/MoviePipelineImageSequenceOutput.h
@@ -76,6 +76,8 @@ public:
 	{
 		OutputFormat = EImageFormat::BMP;
 	}
+
+	virtual void UpdateTelemetry(FMoviePipelineShotRenderTelemetry* InTelemetry) const override;
 };
 
 UCLASS(MinimalAPI)
@@ -95,6 +97,8 @@ public:
 		bWriteAlpha = true;
 	}
 
+	virtual void UpdateTelemetry(FMoviePipelineShotRenderTelemetry* InTelemetry) const override;
+
 	UPROPERTY(EditAnywhere, BlueprintReadWrite, Category = "PNG")
 	bool bWriteAlpha;
 };
@@ -112,6 +116,8 @@ public:
 	{
 		OutputFormat = EImageFormat::JPEG;
 	}
+
+	virtual void UpdateTelemetry(FMoviePipelineShotRenderTelemetry* InTelemetry) const override;
 };
 
 #undef UE_API
