MoviePipeline: Added default names to PPMs that are available by default on renderer nodes in the graph.

Also exposed the ability to change the renderer name and renderer sub name on the renderer nodes (corresponds to the {renderer_name} and {renderer_sub_name} tokens).

Minor: Updated the renderer name for the UI renderer from "UI Renderer" to "UIRenderer" to prevent spaces in filenames.

diff --git a/Engine/Plugins/MovieScene/MoviePipelineMaskRenderPass/Source/MoviePipelineMaskRenderPass/Private/Graph/MovieGraphObjectIdNode.cpp b/Engine/Plugins/MovieScene/MoviePipelineMaskRenderPass/Source/MoviePipelineMaskRenderPass/Private/Graph/MovieGraphObjectIdNode.cpp
--- a/Engine/Plugins/MovieScene/MoviePipelineMaskRenderPass/Source/MoviePipelineMaskRenderPass/Private/Graph/MovieGraphObjectIdNode.cpp
+++ b/Engine/Plugins/MovieScene/MoviePipelineMaskRenderPass/Source/MoviePipelineMaskRenderPass/Private/Graph/MovieGraphObjectIdNode.cpp
@@ -8,6 +8,14 @@
 
 #include UE_INLINE_GENERATED_CPP_BY_NAME(MovieGraphObjectIdNode)
 
+UMovieGraphObjectIdNode::UMovieGraphObjectIdNode()
+{
+	RendererName = TEXT("ObjectID");
+	
+	// The Object ID pass dynamically generates the renderer sub name. This value won't be used.
+	RendererSubName = TEXT("DYNAMIC");
+}
+
 FEngineShowFlags UMovieGraphObjectIdNode::GetShowFlags() const
 {
 	FEngineShowFlags Flags(ESFIM_Game);
@@ -59,6 +67,13 @@ FSlateIcon UMovieGraphObjectIdNode::GetIconAndTint(FLinearColor& OutColor) const
 	
 	return DeferredRendererIcon;
 }
+
+TArray<TPair<FName, UClass*>> UMovieGraphObjectIdNode::GetHiddenProperties() const
+{
+	return {
+		{ TEXT("RendererSubName"), UMovieGraphRenderPassNode::StaticClass() }
+	};
+}
 #endif	// WITH_EDITOR
 
 void UMovieGraphObjectIdNode::UpdateTelemetry(FMoviePipelineShotRenderTelemetry* InTelemetry) const
@@ -71,11 +86,6 @@ TUniquePtr<UE::MovieGraph::Rendering::FMovieGraphImagePassBase> UMovieGraphObjec
 	return MakeUnique<FMovieGraphObjectIdPass>();
 }
 
-FString UMovieGraphObjectIdNode::GetRendererNameImpl() const
-{
-	return TEXT("ObjectID");
-}
-
 void UMovieGraphObjectIdNode::SetupImpl(const FMovieGraphRenderPassSetupData& InSetupData)
 {
 	Super::SetupImpl(InSetupData);
diff --git a/Engine/Plugins/MovieScene/MoviePipelineMaskRenderPass/Source/MoviePipelineMaskRenderPass/Private/Graph/MovieGraphObjectIdPass.cpp b/Engine/Plugins/MovieScene/MoviePipelineMaskRenderPass/Source/MoviePipelineMaskRenderPass/Private/Graph/MovieGraphObjectIdPass.cpp
--- a/Engine/Plugins/MovieScene/MoviePipelineMaskRenderPass/Source/MoviePipelineMaskRenderPass/Private/Graph/MovieGraphObjectIdPass.cpp
+++ b/Engine/Plugins/MovieScene/MoviePipelineMaskRenderPass/Source/MoviePipelineMaskRenderPass/Private/Graph/MovieGraphObjectIdPass.cpp
@@ -198,7 +198,7 @@ void FMovieGraphObjectIdPass::Setup(TWeakObjectPtr<UMovieGraphDefaultRenderer> I
 
 		NewIdentifier.RootBranchName = LayerData.BranchName;
 		NewIdentifier.LayerName = LayerData.LayerName;
-		NewIdentifier.RendererName = InRenderPassNode->GetRendererName();
+		NewIdentifier.RendererName = InLayer.RenderPassNode->GetRendererName();
 		NewIdentifier.CameraName = LayerData.CameraName;
 		NewIdentifier.SubResourceName = UE::MovieGraph::GetCryptomatteTypename(NewIdentifier, InLayer.NumCameras > 1) + FString::Printf(TEXT("%02d"), Index);
 		
diff --git a/Engine/Plugins/MovieScene/MoviePipelineMaskRenderPass/Source/MoviePipelineMaskRenderPass/Public/Graph/MovieGraphObjectIdNode.h b/Engine/Plugins/MovieScene/MoviePipelineMaskRenderPass/Source/MoviePipelineMaskRenderPass/Public/Graph/MovieGraphObjectIdNode.h
--- a/Engine/Plugins/MovieScene/MoviePipelineMaskRenderPass/Source/MoviePipelineMaskRenderPass/Public/Graph/MovieGraphObjectIdNode.h
+++ b/Engine/Plugins/MovieScene/MoviePipelineMaskRenderPass/Source/MoviePipelineMaskRenderPass/Public/Graph/MovieGraphObjectIdNode.h
@@ -18,7 +18,7 @@ class UMovieGraphObjectIdNode final : public UMovieGraphImagePassBaseNode
 	GENERATED_BODY()
 
 public:
-	UMovieGraphObjectIdNode() = default;
+	UMovieGraphObjectIdNode();
 
 	// UMovieGraphImagePassBaseNode Interface
 	UE_API virtual FEngineShowFlags GetShowFlags() const override;
@@ -37,6 +37,7 @@ public:
 #if WITH_EDITOR
 	UE_API virtual FText GetNodeTitle(const bool bGetDescriptive = false) const override;
 	UE_API virtual FSlateIcon GetIconAndTint(FLinearColor& OutColor) const override;
+	UE_API virtual TArray<TPair<FName, UClass*>> GetHiddenProperties() const override;
 #endif
 	// ~UMovieGraphNode Interface
 
@@ -48,7 +49,6 @@ protected:
 	// ~UMovieGraphImagePassBaseNode Interface
 
 	// UMovieGraphRenderPassNode Interface
-	UE_API virtual FString GetRendererNameImpl() const override;
 	UE_API virtual void SetupImpl(const FMovieGraphRenderPassSetupData& InSetupData) override;
 	UE_API virtual void TeardownImpl() override;
 	// ~UMovieGraphRenderPassNode Interface
diff --git a/Engine/Plugins/MovieScene/MovieRenderPipeline/Source/MovieRenderPipelineCore/Private/Graph/MovieGraphDefaultRenderer.cpp b/Engine/Plugins/MovieScene/MovieRenderPipeline/Source/MovieRenderPipelineCore/Private/Graph/MovieGraphDefaultRenderer.cpp
--- a/Engine/Plugins/MovieScene/MovieRenderPipeline/Source/MovieRenderPipelineCore/Private/Graph/MovieGraphDefaultRenderer.cpp
+++ b/Engine/Plugins/MovieScene/MovieRenderPipeline/Source/MovieRenderPipelineCore/Private/Graph/MovieGraphDefaultRenderer.cpp
@@ -3,7 +3,6 @@
 #include "Graph/MovieGraphDefaultRenderer.h"
 #include "Graph/MovieGraphPipeline.h"
 #include "Graph/MovieGraphConfig.h"
-#include "Graph/Nodes/MovieGraphBurnInNode.h"
 #include "Graph/Nodes/MovieGraphDebugNode.h"
 #include "Graph/Nodes/MovieGraphGlobalGameOverrides.h"
 #include "Graph/Nodes/MovieGraphRenderLayerNode.h"
@@ -11,7 +10,6 @@
 #include "Graph/Nodes/MovieGraphGlobalGameOverrides.h"
 #include "Graph/Nodes/MovieGraphDebugNode.h"
 #include "Graph/Nodes/MovieGraphCameraNode.h"
-#include "Graph/Nodes/MovieGraphUIRendererNode.h"
 #include "MovieRenderPipelineCoreModule.h"
 #include "RenderingThread.h"
 #include "Engine/TextureRenderTarget2D.h"
@@ -631,8 +629,8 @@ TArray<FMovieGraphImagePreviewData> UMovieGraphDefaultRenderer::GetPreviewData()
 	{
 		const FString RendererName = RenderTarget.Key.Identifier.RendererName;
 		
-		// Skip the burn-in and widget renderer outputs. They clog up the preview window and seeing them doesn't provide any value.
-		if ((RendererName == UMovieGraphBurnInNode::RendererName) || (RendererName == UMovieGraphUIRendererNode::RendererName))
+		// Skip render targets that have opted out of showing in the preview window.
+		if (!RenderTarget.Key.RenderInitParams.bIncludeInPreviewWindow)
 		{
 			continue;
 		}
diff --git a/Engine/Plugins/MovieScene/MovieRenderPipeline/Source/MovieRenderPipelineCore/Private/Graph/MovieGraphNode.cpp b/Engine/Plugins/MovieScene/MovieRenderPipeline/Source/MovieRenderPipelineCore/Private/Graph/MovieGraphNode.cpp
--- a/Engine/Plugins/MovieScene/MovieRenderPipeline/Source/MovieRenderPipelineCore/Private/Graph/MovieGraphNode.cpp
+++ b/Engine/Plugins/MovieScene/MovieRenderPipeline/Source/MovieRenderPipelineCore/Private/Graph/MovieGraphNode.cpp
@@ -454,6 +454,12 @@ FSlateIcon UMovieGraphNode::GetIconAndTint(FLinearColor& OutColor) const
 	OutColor = FLinearColor::White;
 	return FSlateIcon();
 }
+
+TArray<TPair<FName, UClass*>> UMovieGraphNode::GetHiddenProperties() const
+{
+	// By default, no properties are hidden.
+	return {};
+}
 #endif // WITH_EDITOR
 
 void UMovieGraphNode::PostLoad()
diff --git a/Engine/Plugins/MovieScene/MovieRenderPipeline/Source/MovieRenderPipelineCore/Private/Graph/Nodes/MovieGraphBurnInNode.cpp b/Engine/Plugins/MovieScene/MovieRenderPipeline/Source/MovieRenderPipelineCore/Private/Graph/Nodes/MovieGraphBurnInNode.cpp
--- a/Engine/Plugins/MovieScene/MovieRenderPipeline/Source/MovieRenderPipelineCore/Private/Graph/Nodes/MovieGraphBurnInNode.cpp
+++ b/Engine/Plugins/MovieScene/MovieRenderPipeline/Source/MovieRenderPipelineCore/Private/Graph/Nodes/MovieGraphBurnInNode.cpp
@@ -15,12 +15,20 @@
 
 #define LOCTEXT_NAMESPACE "MovieGraphNodes"
 
+PRAGMA_DISABLE_DEPRECATION_WARNINGS
 const FString UMovieGraphBurnInNode::RendererName = FString("BurnIn");
+PRAGMA_ENABLE_DEPRECATION_WARNINGS
+
 const FString UMovieGraphBurnInNode::DefaultBurnInWidgetAsset = TEXT("/MovieRenderPipeline/Blueprints/Graph/DefaultGraphBurnIn.DefaultGraphBurnIn_C");
 
 UMovieGraphBurnInNode::UMovieGraphBurnInNode()
 	: BurnInClass(DefaultBurnInWidgetAsset)
 {
+	PRAGMA_DISABLE_DEPRECATION_WARNINGS
+	Super::RendererName = RendererName;
+	PRAGMA_ENABLE_DEPRECATION_WARNINGS
+	
+	RendererSubName = TEXT("burnin");
 }
 
 #if WITH_EDITOR
@@ -149,7 +157,7 @@ void UMovieGraphBurnInNode::FMovieGraphBurnInPass::Setup(TWeakObjectPtr<UMovieGr
 {
 	FMovieGraphWidgetPass::Setup(InRenderer, InLayer);
 	
-	RenderDataIdentifier.SubResourceName = TEXT("burnin");
+	RenderDataIdentifier.SubResourceName = InLayer.RenderPassNode->GetRendererSubName();
 
 	UMovieGraphBurnInNode* ParentNodeOnInitialization = Cast<UMovieGraphBurnInNode>(InLayer.RenderPassNode);
 	if (!ensureMsgf(ParentNodeOnInitialization, TEXT("FMovieGraphBurnInPass shouldn't exist without a parent node in the graph.")))
diff --git a/Engine/Plugins/MovieScene/MovieRenderPipeline/Source/MovieRenderPipelineCore/Private/Graph/Nodes/MovieGraphUIRendererNode.cpp b/Engine/Plugins/MovieScene/MovieRenderPipeline/Source/MovieRenderPipelineCore/Private/Graph/Nodes/MovieGraphUIRendererNode.cpp
--- a/Engine/Plugins/MovieScene/MovieRenderPipeline/Source/MovieRenderPipelineCore/Private/Graph/Nodes/MovieGraphUIRendererNode.cpp
+++ b/Engine/Plugins/MovieScene/MovieRenderPipeline/Source/MovieRenderPipelineCore/Private/Graph/Nodes/MovieGraphUIRendererNode.cpp
@@ -11,7 +11,18 @@
 
 #include UE_INLINE_GENERATED_CPP_BY_NAME(MovieGraphUIRendererNode)
 
-const FString UMovieGraphUIRendererNode::RendererName = FString("UI Renderer");
+PRAGMA_DISABLE_DEPRECATION_WARNINGS
+const FString UMovieGraphUIRendererNode::RendererName = FString("UIRenderer");
+PRAGMA_ENABLE_DEPRECATION_WARNINGS
+
+UMovieGraphUIRendererNode::UMovieGraphUIRendererNode()
+{
+	PRAGMA_DISABLE_DEPRECATION_WARNINGS
+	Super::RendererName = RendererName;
+	PRAGMA_ENABLE_DEPRECATION_WARNINGS
+	
+	RendererSubName = TEXT("ui");
+}
 
 #if WITH_EDITOR
 FText UMovieGraphUIRendererNode::GetNodeTitle(const bool bGetDescriptive) const
@@ -37,7 +48,7 @@ void UMovieGraphUIRendererNode::FMovieGraphUIPass::Setup(TWeakObjectPtr<UMovieGr
 {
 	FMovieGraphWidgetPass::Setup(InRenderer, InLayer);
 
-	RenderDataIdentifier.SubResourceName = TEXT("ui");
+	RenderDataIdentifier.SubResourceName = InLayer.RenderPassNode->GetRendererSubName();
 }
 
 TSharedPtr<SWidget> UMovieGraphUIRendererNode::FMovieGraphUIPass::GetWidget(UMovieGraphWidgetRendererBaseNode* InNodeThisFrame)
diff --git a/Engine/Plugins/MovieScene/MovieRenderPipeline/Source/MovieRenderPipelineCore/Private/Graph/Nodes/MovieGraphWidgetRendererBaseNode.cpp b/Engine/Plugins/MovieScene/MovieRenderPipeline/Source/MovieRenderPipelineCore/Private/Graph/Nodes/MovieGraphWidgetRendererBaseNode.cpp
--- a/Engine/Plugins/MovieScene/MovieRenderPipeline/Source/MovieRenderPipelineCore/Private/Graph/Nodes/MovieGraphWidgetRendererBaseNode.cpp
+++ b/Engine/Plugins/MovieScene/MovieRenderPipeline/Source/MovieRenderPipelineCore/Private/Graph/Nodes/MovieGraphWidgetRendererBaseNode.cpp
@@ -91,7 +91,7 @@ void UMovieGraphWidgetRendererBaseNode::FMovieGraphWidgetPass::Setup(TWeakObject
 
 	RenderDataIdentifier.RootBranchName = LayerData.BranchName;
 	RenderDataIdentifier.LayerName = LayerData.LayerName;
-	RenderDataIdentifier.RendererName = CastChecked<UMovieGraphWidgetRendererBaseNode>(InLayer.RenderPassNode)->GetRendererName();
+	RenderDataIdentifier.RendererName = InLayer.RenderPassNode->GetRendererName();
 
 	// NOTE: Subclasses should specify the SubResourceName
 
@@ -153,6 +153,7 @@ void UMovieGraphWidgetRendererBaseNode::FMovieGraphWidgetPass::Render(const FMov
 	RenderTargetInitParams.Size = OutputResolution;
 	RenderTargetInitParams.TargetGamma = UTextureRenderTarget::GetDefaultDisplayGamma();
 	RenderTargetInitParams.PixelFormat = PF_B8G8R8A8;
+	RenderTargetInitParams.bIncludeInPreviewWindow = false;
 	UTextureRenderTarget2D* RenderTarget = Renderer->GetOrCreateViewRenderTarget(RenderTargetInitParams, RenderDataIdentifier);
 
 	if (InFrameTraversalContext.Time.bIsFirstTemporalSampleForFrame)
diff --git a/Engine/Plugins/MovieScene/MovieRenderPipeline/Source/MovieRenderPipelineCore/Public/Graph/MovieGraphDataTypes.h b/Engine/Plugins/MovieScene/MovieRenderPipeline/Source/MovieRenderPipelineCore/Public/Graph/MovieGraphDataTypes.h
--- a/Engine/Plugins/MovieScene/MovieRenderPipeline/Source/MovieRenderPipelineCore/Public/Graph/MovieGraphDataTypes.h
+++ b/Engine/Plugins/MovieScene/MovieRenderPipeline/Source/MovieRenderPipelineCore/Public/Graph/MovieGraphDataTypes.h
@@ -334,7 +334,11 @@ struct FMovieGraphRenderPassLayerData
 
 	// The number of cameras that are being rendered.
 	uint32 NumCameras = 1;
-	
+
+	/**
+	 * The render pass node that's associated with this layer data. WARNING: this node will ONLY be valid during initial pass setup! Post-setup, if
+	 * a non-CDO node is needed, it needs to be fetched from the evaluated config.
+	 */
 	TWeakObjectPtr<class UMovieGraphRenderPassNode> RenderPassNode;
 };
 // ToDo: Both of these can probably go into the Default Renderer implementation.
diff --git a/Engine/Plugins/MovieScene/MovieRenderPipeline/Source/MovieRenderPipelineCore/Public/Graph/MovieGraphDefaultRenderer.h b/Engine/Plugins/MovieScene/MovieRenderPipeline/Source/MovieRenderPipelineCore/Public/Graph/MovieGraphDefaultRenderer.h
--- a/Engine/Plugins/MovieScene/MovieRenderPipeline/Source/MovieRenderPipelineCore/Public/Graph/MovieGraphDefaultRenderer.h
+++ b/Engine/Plugins/MovieScene/MovieRenderPipeline/Source/MovieRenderPipelineCore/Public/Graph/MovieGraphDefaultRenderer.h
@@ -90,6 +90,7 @@ namespace UE::MovieGraph::DefaultRenderer
 			, TargetGamma(0.f)
 			, bForceLinearGamma(false)
 			, PixelFormat(EPixelFormat::PF_Unknown)
+			, bIncludeInPreviewWindow(true)
 		{
 		}
 
@@ -98,12 +99,16 @@ namespace UE::MovieGraph::DefaultRenderer
 		bool bForceLinearGamma;
 		EPixelFormat PixelFormat;
 
+		/** Whether this render target is shown in the "preview" window that displays during a render. */
+		bool bIncludeInPreviewWindow;
+
 		bool operator == (const FRenderTargetInitParams& InRHS) const
 		{
 			return Size == InRHS.Size
 				&& TargetGamma == InRHS.TargetGamma
 				&& bForceLinearGamma == InRHS.bForceLinearGamma
-				&& PixelFormat == InRHS.PixelFormat;
+				&& PixelFormat == InRHS.PixelFormat
+				&& bIncludeInPreviewWindow == InRHS.bIncludeInPreviewWindow;
 		}
 
 		bool operator != (const FRenderTargetInitParams& InRHS) const
@@ -113,7 +118,11 @@ namespace UE::MovieGraph::DefaultRenderer
 
 		friend uint32 GetTypeHash(FRenderTargetInitParams Params)
 		{
-			return HashCombineFast(GetTypeHash(Params.Size), HashCombineFast(GetTypeHash(Params.TargetGamma), HashCombineFast(GetTypeHash(Params.bForceLinearGamma), GetTypeHash(Params.PixelFormat))));
+			return HashCombineFast(GetTypeHash(Params.Size),
+				HashCombineFast(GetTypeHash(Params.TargetGamma),
+				HashCombineFast(GetTypeHash(Params.bForceLinearGamma),
+				HashCombineFast(GetTypeHash(Params.PixelFormat),
+				GetTypeHash(Params.bIncludeInPreviewWindow)))));
 		}
 	};
 
diff --git a/Engine/Plugins/MovieScene/MovieRenderPipeline/Source/MovieRenderPipelineCore/Public/Graph/MovieGraphNode.h b/Engine/Plugins/MovieScene/MovieRenderPipeline/Source/MovieRenderPipelineCore/Public/Graph/MovieGraphNode.h
--- a/Engine/Plugins/MovieScene/MovieRenderPipeline/Source/MovieRenderPipelineCore/Public/Graph/MovieGraphNode.h
+++ b/Engine/Plugins/MovieScene/MovieRenderPipeline/Source/MovieRenderPipelineCore/Public/Graph/MovieGraphNode.h
@@ -350,6 +350,15 @@ public:
 	bool IsCommentBubbleVisible() const { return bIsCommentBubbleVisible; }
 	UFUNCTION(BlueprintCallable, Category = "Movie Graph")
 	void SetIsCommentBubbleVisible(uint8 bIsVisible) { bIsCommentBubbleVisible = bIsVisible; }
+
+	/**
+	 * Gets the properties on this node that should be hidden by the base UMovieGraphNode customization. The FName in the pair is the property
+	 * name, and the UClass* is the class that the property is defined on (may be a base class).
+	 *
+	 * Note that other customizations can easily override this behavior, but it's useful in cases where (for example) an inherited
+	 * property should be hidden, and there aren't any other customizations needed.
+	 */
+	MOVIERENDERPIPELINECORE_API virtual TArray<TPair<FName, UClass*>> GetHiddenProperties() const;
 #endif
 
 	//~ Begin UObject Interface
diff --git a/Engine/Plugins/MovieScene/MovieRenderPipeline/Source/MovieRenderPipelineCore/Public/Graph/Nodes/MovieGraphBurnInNode.h b/Engine/Plugins/MovieScene/MovieRenderPipeline/Source/MovieRenderPipelineCore/Public/Graph/Nodes/MovieGraphBurnInNode.h
--- a/Engine/Plugins/MovieScene/MovieRenderPipeline/Source/MovieRenderPipelineCore/Public/Graph/Nodes/MovieGraphBurnInNode.h
+++ b/Engine/Plugins/MovieScene/MovieRenderPipeline/Source/MovieRenderPipelineCore/Public/Graph/Nodes/MovieGraphBurnInNode.h
@@ -19,7 +19,7 @@ struct FMovieGraphRenderPassLayerData;
  *
  * DEPRECATED. Use the burn-in settings on output nodes instead to generate burn-ins.
  */
-UCLASS(MinimalAPI, Hidden, HideDropdown)
+UCLASS(MinimalAPI, Hidden, HideDropdown, HideCategories = (Naming))
 class UMovieGraphBurnInNode : public UMovieGraphWidgetRendererBaseNode
 {
 	GENERATED_BODY()
@@ -40,6 +40,7 @@ public:
 	UE_API TObjectPtr<UMovieGraphBurnInWidget> GetOrCreateBurnInWidget(UClass* WidgetClass, UWorld* InOwner);
 
 public:
+	UE_DEPRECATED(5.7, "RendererName usage is deprecated, please use GetRendererName() instead.")
 	static UE_API const FString RendererName;
 
 	/** The path to the default widget class that will be used for the burn-in. */
@@ -89,7 +90,6 @@ protected:
 	
 	// UMovieGraphRenderPassNode Interface
 	UE_API virtual void GatherOutputPassesImpl(UMovieGraphEvaluatedConfig* InConfig, TArray<FMovieGraphRenderDataIdentifier>& OutExpectedPasses) const override;
-	virtual FString GetRendererNameImpl() const override { return RendererName; }
 	UE_API virtual void TeardownImpl() override;
 	// ~UMovieGraphRenderPassNode Interface
 };
diff --git a/Engine/Plugins/MovieScene/MovieRenderPipeline/Source/MovieRenderPipelineCore/Public/Graph/Nodes/MovieGraphRenderPassNode.h b/Engine/Plugins/MovieScene/MovieRenderPipeline/Source/MovieRenderPipelineCore/Public/Graph/Nodes/MovieGraphRenderPassNode.h
--- a/Engine/Plugins/MovieScene/MovieRenderPipeline/Source/MovieRenderPipelineCore/Public/Graph/Nodes/MovieGraphRenderPassNode.h
+++ b/Engine/Plugins/MovieScene/MovieRenderPipeline/Source/MovieRenderPipelineCore/Public/Graph/Nodes/MovieGraphRenderPassNode.h
@@ -23,13 +23,22 @@ UCLASS(MinimalAPI, Abstract)
 class UMovieGraphRenderPassNode : public UMovieGraphSettingNode
 {
 	GENERATED_BODY()
+
 public:
 	UMovieGraphRenderPassNode()
 	{
+		RendererName = TEXT("UnnamedRenderPass");
+		RendererSubName = TEXT("beauty");
 	}
 
-	/** Get the name of this renderer. Deferred, Path Tracer, Panoramic, etc. Called on the CDO. */
+	/** Get the name of this renderer. Deferred, Path Tracer, Panoramic, etc. */
 	FString GetRendererName() const { return GetRendererNameImpl(); }
+
+	/**
+	 * Get the sub-name of this renderer, which identifies a specific pass that this renderer generates. Each type of pass that this
+	 * renderer generates should have a unique sub-name.
+	 */
+	FString GetRendererSubName() const { return GetRendererSubNameImpl(); }
 	
 	/** Called when this should set up for rendering a new shot. Called on the CDO. */
 	void Setup(const FMovieGraphRenderPassSetupData& InSetupData) { SetupImpl(InSetupData); }
@@ -80,10 +89,26 @@ public:
 		OutColor = FLinearColor::White;
 		return DeferredRendererIcon;
 	}
-#endif
+#endif	// WITH_EDITOR
+
+protected:
+	UPROPERTY(EditAnywhere, BlueprintReadWrite, Category = Overrides, meta = (InlineEditConditionToggle))
+	uint8 bOverride_RendererName : 1;
+
+	UPROPERTY(EditAnywhere, BlueprintReadWrite, Category = Overrides, meta = (InlineEditConditionToggle))
+	uint8 bOverride_RendererSubName : 1;
+
+	/** The value that will be used in the {renderer_name} token. */
+	UPROPERTY(EditAnywhere, BlueprintReadWrite, Category = "Naming", meta = (EditCondition="bOverride_RendererName"))
+	FString RendererName;
+
+	/** The value that will be used in the {renderer_sub_name} token. */
+	UPROPERTY(EditAnywhere, BlueprintReadWrite, Category = "Naming", meta = (EditCondition="bOverride_RendererSubName"))
+	FString RendererSubName;
 
 protected:
-	virtual FString GetRendererNameImpl() const { return TEXT("UnnamedRenderPass"); }
+	virtual FString GetRendererNameImpl() const { return RendererName; }
+	virtual FString GetRendererSubNameImpl() const { return RendererSubName; }
 	virtual void SetupImpl(const FMovieGraphRenderPassSetupData& InSetupData) {}
 	virtual void TeardownImpl() {}
 	virtual void RenderImpl(const FMovieGraphTraversalContext& InFrameTraversalContext, const FMovieGraphTimeStepData& InTimeData) {}
diff --git a/Engine/Plugins/MovieScene/MovieRenderPipeline/Source/MovieRenderPipelineCore/Public/Graph/Nodes/MovieGraphUIRendererNode.h b/Engine/Plugins/MovieScene/MovieRenderPipeline/Source/MovieRenderPipelineCore/Public/Graph/Nodes/MovieGraphUIRendererNode.h
--- a/Engine/Plugins/MovieScene/MovieRenderPipeline/Source/MovieRenderPipelineCore/Public/Graph/Nodes/MovieGraphUIRendererNode.h
+++ b/Engine/Plugins/MovieScene/MovieRenderPipeline/Source/MovieRenderPipelineCore/Public/Graph/Nodes/MovieGraphUIRendererNode.h
@@ -21,6 +21,8 @@ class UMovieGraphUIRendererNode : public UMovieGraphWidgetRendererBaseNode
 	GENERATED_BODY()
 
 public:
+	UMovieGraphUIRendererNode();
+
 #if WITH_EDITOR
 	UE_API virtual FText GetNodeTitle(const bool bGetDescriptive = false) const override;
 	UE_API virtual FLinearColor GetNodeTitleColor() const override;
@@ -28,6 +30,7 @@ public:
 #endif
 
 public:
+	UE_DEPRECATED(5.7, "RendererName usage is deprecated, please use GetRendererName() instead.")
 	static UE_API const FString RendererName;
 
 protected:
@@ -42,10 +45,6 @@ protected:
 	// UMovieGraphWidgetRendererBaseNode Interface
 	UE_API virtual TUniquePtr<FMovieGraphWidgetPass> GeneratePass() override;
 	// ~UMovieGraphWidgetRendererBaseNode Interface
-	
-	// UMovieGraphRenderPassNode Interface
-	virtual FString GetRendererNameImpl() const override { return RendererName; }
-	// ~UMovieGraphRenderPassNode Interface
 };
 
 #undef UE_API
diff --git a/Engine/Plugins/MovieScene/MovieRenderPipeline/Source/MovieRenderPipelineEditor/Private/Customizations/Graph/MovieGraphNodeCustomization.h b/Engine/Plugins/MovieScene/MovieRenderPipeline/Source/MovieRenderPipelineEditor/Private/Customizations/Graph/MovieGraphNodeCustomization.h
--- a/Engine/Plugins/MovieScene/MovieRenderPipeline/Source/MovieRenderPipelineEditor/Private/Customizations/Graph/MovieGraphNodeCustomization.h
+++ b/Engine/Plugins/MovieScene/MovieRenderPipeline/Source/MovieRenderPipelineEditor/Private/Customizations/Graph/MovieGraphNodeCustomization.h
@@ -33,6 +33,8 @@ protected:
 
 		// Hide the "Properties" category, which houses the dynamic properties, if there are no dynamic properties
 		// on the node. An empty category in the details panel is messy/confusing.
+		//
+		// Some nodes may also request that specific properties (most likely inherited) are hidden.
 		for (const TWeakObjectPtr<UObject>& CustomizedObject : ObjectsBeingCustomized)
 		{
 			if (const UMovieGraphNode* Node = Cast<UMovieGraphNode>(CustomizedObject))
@@ -40,7 +42,11 @@ protected:
 				if (Node->GetDynamicPropertyDescriptions().IsEmpty())
 				{
 					DetailBuilder.HideCategory("Properties");
-					break;
+				}
+
+				for (const TPair<FName, UClass*>& HiddenProperty : Node->GetHiddenProperties())
+				{
+					DetailBuilder.HideProperty(HiddenProperty.Key, HiddenProperty.Value);
 				}
 			}
 		}
@@ -55,7 +61,11 @@ protected:
 		// Collapse the Tags category by default because it's not used by most users. Set the sort order to make Tags appear last in the list of
 		// categories (1000 * Uncommon == starting index for all categories with "Uncommon" sort order, and the +1000 is to move the category past all
 		// other categories with "Uncommon" sort order).
-		DetailBuilder.EditCategory("Tags").InitiallyCollapsed(true).SetSortOrder(1000 * ECategoryPriority::Uncommon + 1000);
+		IDetailCategoryBuilder& TagsCategory = DetailBuilder.EditCategory("Tags");
+		TagsCategory.InitiallyCollapsed(true).SetSortOrder(1000 * ECategoryPriority::Uncommon + 1000);
+
+		// Show the "Naming" category, which is on renderer nodes, right before the "Tags" category.
+		DetailBuilder.EditCategory("Naming").SetSortOrder(TagsCategory.GetSortOrder() - 1);
 
 		// The File Output category should be shown first in all output nodes that have this category.
 		DetailBuilder.EditCategory("File Output").SetSortOrder(ECategoryPriority::Important);
diff --git a/Engine/Plugins/MovieScene/MovieRenderPipeline/Source/MovieRenderPipelineRenderPasses/Private/Graph/Nodes/MovieGraphDeferredPanoramicPassNode.cpp b/Engine/Plugins/MovieScene/MovieRenderPipeline/Source/MovieRenderPipelineRenderPasses/Private/Graph/Nodes/MovieGraphDeferredPanoramicPassNode.cpp
--- a/Engine/Plugins/MovieScene/MovieRenderPipeline/Source/MovieRenderPipelineRenderPasses/Private/Graph/Nodes/MovieGraphDeferredPanoramicPassNode.cpp
+++ b/Engine/Plugins/MovieScene/MovieRenderPipeline/Source/MovieRenderPipelineRenderPasses/Private/Graph/Nodes/MovieGraphDeferredPanoramicPassNode.cpp
@@ -19,7 +19,9 @@ UMovieGraphDeferredPanoramicNode::UMovieGraphDeferredPanoramicNode()
 	, bDisableToneCurve(false)
 	, bAllowOCIO(true)
 	, ViewModeIndex(VMI_Lit)
-{}
+{
+	RendererName = TEXT("DeferredPanoramic");
+}
 
 FEngineShowFlags UMovieGraphDeferredPanoramicNode::GetShowFlags() const
 {
@@ -57,11 +59,6 @@ TUniquePtr<UE::MovieGraph::Rendering::FMovieGraphImagePassBase> UMovieGraphDefer
 	return MakeUnique<UE::MovieGraph::Rendering::FMovieGraphDeferredPanoramicPass>();
 }
 
-FString UMovieGraphDeferredPanoramicNode::GetRendererNameImpl() const
-{
-	return TEXT("DeferredPanoramic");
-}
-
 EViewModeIndex UMovieGraphDeferredPanoramicNode::GetViewModeIndex() const
 {
 	return ViewModeIndex;
diff --git a/Engine/Plugins/MovieScene/MovieRenderPipeline/Source/MovieRenderPipelineRenderPasses/Private/Graph/Nodes/MovieGraphDeferredPassNode.cpp b/Engine/Plugins/MovieScene/MovieRenderPipeline/Source/MovieRenderPipelineRenderPasses/Private/Graph/Nodes/MovieGraphDeferredPassNode.cpp
--- a/Engine/Plugins/MovieScene/MovieRenderPipeline/Source/MovieRenderPipelineRenderPasses/Private/Graph/Nodes/MovieGraphDeferredPassNode.cpp
+++ b/Engine/Plugins/MovieScene/MovieRenderPipeline/Source/MovieRenderPipelineRenderPasses/Private/Graph/Nodes/MovieGraphDeferredPassNode.cpp
@@ -25,6 +25,7 @@ UMovieGraphDeferredRenderPassNode::UMovieGraphDeferredRenderPassNode()
 	, bAllocateHistoryPerTile(false)
 	, bPageToSystemMemory(false)
 {
+	RendererName = TEXT("Deferred");
 
 	// To help user knowledge we pre-seed the additional post processing materials with an array of potentially common passes.
 	TArray<FString> DefaultPostProcessMaterials;
@@ -34,6 +35,7 @@ UMovieGraphDeferredRenderPassNode::UMovieGraphDeferredRenderPassNode()
 	for (FString& MaterialPath : DefaultPostProcessMaterials)
 	{
 		FMoviePipelinePostProcessPass& NewPass = AdditionalPostProcessMaterials.AddDefaulted_GetRef();
+		NewPass.Name = (MaterialPath == DefaultDepthAsset) ? FString(TEXT("depth")) : FString(TEXT("motion"));
 		NewPass.Material = TSoftObjectPtr<UMaterialInterface>(FSoftObjectPath(MaterialPath));
 		NewPass.bEnabled = false;
 		NewPass.bHighPrecisionOutput = MaterialPath.Equals(DefaultDepthAsset);
@@ -83,12 +85,6 @@ FSlateIcon UMovieGraphDeferredRenderPassNode::GetIconAndTint(FLinearColor& OutCo
 }
 #endif
 
-FString UMovieGraphDeferredRenderPassNode::GetRendererNameImpl() const
-{
-	static const FString RendererNameImpl(TEXT("Deferred"));
-	return RendererNameImpl;
-}
-
 EViewModeIndex UMovieGraphDeferredRenderPassNode::GetViewModeIndex() const
 {
 	return ViewModeIndex;
diff --git a/Engine/Plugins/MovieScene/MovieRenderPipeline/Source/MovieRenderPipelineRenderPasses/Private/Graph/Nodes/MovieGraphImagePassBaseNode.cpp b/Engine/Plugins/MovieScene/MovieRenderPipeline/Source/MovieRenderPipelineRenderPasses/Private/Graph/Nodes/MovieGraphImagePassBaseNode.cpp
--- a/Engine/Plugins/MovieScene/MovieRenderPipeline/Source/MovieRenderPipelineRenderPasses/Private/Graph/Nodes/MovieGraphImagePassBaseNode.cpp
+++ b/Engine/Plugins/MovieScene/MovieRenderPipeline/Source/MovieRenderPipelineRenderPasses/Private/Graph/Nodes/MovieGraphImagePassBaseNode.cpp
@@ -21,6 +21,7 @@ FString UMovieGraphImagePassBaseNode::DefaultMotionVectorsAsset = TEXT("/MovieRe
 
 UMovieGraphImagePassBaseNode::UMovieGraphImagePassBaseNode()
 {
+	RendererName = TEXT("UnnamedImageRenderPass");
 	ShowFlags = CreateDefaultSubobject<UMovieGraphShowFlags>(TEXT("ShowFlags"));
 }
 
diff --git a/Engine/Plugins/MovieScene/MovieRenderPipeline/Source/MovieRenderPipelineRenderPasses/Private/Graph/Nodes/MovieGraphPathTracerPassNode.cpp b/Engine/Plugins/MovieScene/MovieRenderPipeline/Source/MovieRenderPipelineRenderPasses/Private/Graph/Nodes/MovieGraphPathTracerPassNode.cpp
--- a/Engine/Plugins/MovieScene/MovieRenderPipeline/Source/MovieRenderPipelineRenderPasses/Private/Graph/Nodes/MovieGraphPathTracerPassNode.cpp
+++ b/Engine/Plugins/MovieScene/MovieRenderPipeline/Source/MovieRenderPipelineRenderPasses/Private/Graph/Nodes/MovieGraphPathTracerPassNode.cpp
@@ -28,6 +28,8 @@ UMovieGraphPathTracerRenderPassNode::UMovieGraphPathTracerRenderPassNode()
 	, bAllowOCIO(true)
 	, TileCount(1)
 {
+	RendererName = TEXT("PathTraced");
+
 	ShowFlags->ApplyDefaultShowFlagValue(VMI_PathTracing, true);
 	// TODO: Showflag for SetMotionBlur()?
 
@@ -39,6 +41,7 @@ UMovieGraphPathTracerRenderPassNode::UMovieGraphPathTracerRenderPassNode()
 	for (FString& MaterialPath : DefaultPostProcessMaterials)
 	{
 		FMoviePipelinePostProcessPass& NewPass = AdditionalPostProcessMaterials.AddDefaulted_GetRef();
+		NewPass.Name = (MaterialPath == DefaultDepthAsset) ? FString(TEXT("depth")) : FString(TEXT("motion"));
 		NewPass.Material = TSoftObjectPtr<UMaterialInterface>(FSoftObjectPath(MaterialPath));
 		NewPass.bEnabled = false;
 		NewPass.bHighPrecisionOutput = MaterialPath.Equals(DefaultDepthAsset);
@@ -178,12 +181,6 @@ void UMovieGraphPathTracerRenderPassNode::TeardownImpl()
 	}
 }
 
-FString UMovieGraphPathTracerRenderPassNode::GetRendererNameImpl() const
-{
-	static const FString RendererNameImpl(TEXT("PathTraced"));
-	return RendererNameImpl;
-}
-
 int32 UMovieGraphPathTracerRenderPassNode::GetCoolingDownFrameCount() const
 {
 	return ((DenoiserType == EMovieGraphPathTracerDenoiserType::Temporal)
diff --git a/Engine/Plugins/MovieScene/MovieRenderPipeline/Source/MovieRenderPipelineRenderPasses/Private/Graph/Renderers/MovieGraphDeferredPanoramicPass.cpp b/Engine/Plugins/MovieScene/MovieRenderPipeline/Source/MovieRenderPipelineRenderPasses/Private/Graph/Renderers/MovieGraphDeferredPanoramicPass.cpp
--- a/Engine/Plugins/MovieScene/MovieRenderPipeline/Source/MovieRenderPipelineRenderPasses/Private/Graph/Renderers/MovieGraphDeferredPanoramicPass.cpp
+++ b/Engine/Plugins/MovieScene/MovieRenderPipeline/Source/MovieRenderPipelineRenderPasses/Private/Graph/Renderers/MovieGraphDeferredPanoramicPass.cpp
@@ -266,8 +266,8 @@ void FMovieGraphDeferredPanoramicPass::Setup(TWeakObjectPtr<UMovieGraphDefaultRe
 
 	RenderDataIdentifier.RootBranchName = InLayer.BranchName;
 	RenderDataIdentifier.LayerName = InLayer.LayerName;
-	RenderDataIdentifier.RendererName = InRenderPassNode->GetRendererName();
-	RenderDataIdentifier.SubResourceName = TEXT("beauty");
+	RenderDataIdentifier.RendererName = InLayer.RenderPassNode->GetRendererName();
+	RenderDataIdentifier.SubResourceName = InLayer.RenderPassNode->GetRendererSubName();
 	RenderDataIdentifier.CameraName = InLayer.CameraName;
 	LayerData = InLayer;
 
diff --git a/Engine/Plugins/MovieScene/MovieRenderPipeline/Source/MovieRenderPipelineRenderPasses/Private/Graph/Renderers/MovieGraphDeferredPass.cpp b/Engine/Plugins/MovieScene/MovieRenderPipeline/Source/MovieRenderPipelineRenderPasses/Private/Graph/Renderers/MovieGraphDeferredPass.cpp
--- a/Engine/Plugins/MovieScene/MovieRenderPipeline/Source/MovieRenderPipelineRenderPasses/Private/Graph/Renderers/MovieGraphDeferredPass.cpp
+++ b/Engine/Plugins/MovieScene/MovieRenderPipeline/Source/MovieRenderPipelineRenderPasses/Private/Graph/Renderers/MovieGraphDeferredPass.cpp
@@ -59,8 +59,8 @@ void FMovieGraphDeferredPass::Setup(TWeakObjectPtr<UMovieGraphDefaultRenderer> I
 
 	RenderDataIdentifier.RootBranchName = LayerData.BranchName;
 	RenderDataIdentifier.LayerName = LayerData.LayerName;
-	RenderDataIdentifier.RendererName = InRenderPassNode->GetRendererName();
-	RenderDataIdentifier.SubResourceName = TEXT("beauty");
+	RenderDataIdentifier.RendererName = InLayer.RenderPassNode->GetRendererName();
+	RenderDataIdentifier.SubResourceName = InLayer.RenderPassNode->GetRendererSubName();
 	RenderDataIdentifier.CameraName =  InLayer.CameraName;
 
 	UMovieGraphImagePassBaseNode* ParentNode = Cast< UMovieGraphImagePassBaseNode>(LayerData.RenderPassNode);
diff --git a/Engine/Plugins/MovieScene/MovieRenderPipeline/Source/MovieRenderPipelineRenderPasses/Public/Graph/Nodes/MovieGraphDeferredPanoramicPassNode.h b/Engine/Plugins/MovieScene/MovieRenderPipeline/Source/MovieRenderPipelineRenderPasses/Public/Graph/Nodes/MovieGraphDeferredPanoramicPassNode.h
--- a/Engine/Plugins/MovieScene/MovieRenderPipeline/Source/MovieRenderPipelineRenderPasses/Public/Graph/Nodes/MovieGraphDeferredPanoramicPassNode.h
+++ b/Engine/Plugins/MovieScene/MovieRenderPipeline/Source/MovieRenderPipelineRenderPasses/Public/Graph/Nodes/MovieGraphDeferredPanoramicPassNode.h
@@ -19,7 +19,6 @@ public:
 	// UMovieGraphImagePassBaseNode Interface
 	UE_API virtual FEngineShowFlags GetShowFlags() const override;
 	UE_API virtual EViewModeIndex GetViewModeIndex() const override;
-	UE_API virtual FString GetRendererNameImpl() const override;
 	UE_API virtual TUniquePtr<UE::MovieGraph::Rendering::FMovieGraphImagePassBase> CreateInstance() const override;
 	UE_API virtual bool GetWriteAllSamples() const override;
 	UE_API virtual int32 GetNumSpatialSamples() const override;
diff --git a/Engine/Plugins/MovieScene/MovieRenderPipeline/Source/MovieRenderPipelineRenderPasses/Public/Graph/Nodes/MovieGraphDeferredPassNode.h b/Engine/Plugins/MovieScene/MovieRenderPipeline/Source/MovieRenderPipelineRenderPasses/Public/Graph/Nodes/MovieGraphDeferredPassNode.h
--- a/Engine/Plugins/MovieScene/MovieRenderPipeline/Source/MovieRenderPipelineRenderPasses/Public/Graph/Nodes/MovieGraphDeferredPassNode.h
+++ b/Engine/Plugins/MovieScene/MovieRenderPipeline/Source/MovieRenderPipelineRenderPasses/Public/Graph/Nodes/MovieGraphDeferredPassNode.h
@@ -28,7 +28,6 @@ public:
 
 protected:
 	// UMovieGraphRenderPassNode Interface
-	UE_API virtual FString GetRendererNameImpl() const override;
 	UE_API virtual TUniquePtr<UE::MovieGraph::Rendering::FMovieGraphImagePassBase> CreateInstance() const override;
 	UE_API virtual bool GetWriteAllSamples() const override;
 	UE_API virtual TArray<FMoviePipelinePostProcessPass> GetAdditionalPostProcessMaterials() const override;
diff --git a/Engine/Plugins/MovieScene/MovieRenderPipeline/Source/MovieRenderPipelineRenderPasses/Public/Graph/Nodes/MovieGraphImagePassBaseNode.h b/Engine/Plugins/MovieScene/MovieRenderPipeline/Source/MovieRenderPipelineRenderPasses/Public/Graph/Nodes/MovieGraphImagePassBaseNode.h
--- a/Engine/Plugins/MovieScene/MovieRenderPipeline/Source/MovieRenderPipelineRenderPasses/Public/Graph/Nodes/MovieGraphImagePassBaseNode.h
+++ b/Engine/Plugins/MovieScene/MovieRenderPipeline/Source/MovieRenderPipelineRenderPasses/Public/Graph/Nodes/MovieGraphImagePassBaseNode.h
@@ -111,7 +111,6 @@ protected:
 
 protected:
 	// UMovieGraphRenderPassNode Interface
-	virtual FString GetRendererNameImpl() const { return TEXT("UnnamedImageRenderPass"); }
 	UE_API virtual void SetupImpl(const FMovieGraphRenderPassSetupData& InSetupData) override;
 	UE_API virtual void TeardownImpl() override;
 	UE_API virtual void RenderImpl(const FMovieGraphTraversalContext& InFrameTraversalContext, const FMovieGraphTimeStepData& InTimeData) override;
diff --git a/Engine/Plugins/MovieScene/MovieRenderPipeline/Source/MovieRenderPipelineRenderPasses/Public/Graph/Nodes/MovieGraphPathTracerPassNode.h b/Engine/Plugins/MovieScene/MovieRenderPipeline/Source/MovieRenderPipelineRenderPasses/Public/Graph/Nodes/MovieGraphPathTracerPassNode.h
--- a/Engine/Plugins/MovieScene/MovieRenderPipeline/Source/MovieRenderPipelineRenderPasses/Public/Graph/Nodes/MovieGraphPathTracerPassNode.h
+++ b/Engine/Plugins/MovieScene/MovieRenderPipeline/Source/MovieRenderPipelineRenderPasses/Public/Graph/Nodes/MovieGraphPathTracerPassNode.h
@@ -74,7 +74,6 @@ public:
 
 protected:
 	// UMovieGraphRenderPassNode Interface
-	UE_API virtual FString GetRendererNameImpl() const override;
 	UE_API virtual int32 GetCoolingDownFrameCount() const override;
 	// ~UMovieGraphRenderPassNode Interface
 
diff --git a/Engine/Plugins/MovieScene/MovieRenderPipeline/Source/MovieRenderPipelineRenderPasses/Public/MovieGraphImageSequenceOutputNode.h b/Engine/Plugins/MovieScene/MovieRenderPipeline/Source/MovieRenderPipelineRenderPasses/Public/MovieGraphImageSequenceOutputNode.h
--- a/Engine/Plugins/MovieScene/MovieRenderPipeline/Source/MovieRenderPipelineRenderPasses/Public/MovieGraphImageSequenceOutputNode.h
+++ b/Engine/Plugins/MovieScene/MovieRenderPipeline/Source/MovieRenderPipelineRenderPasses/Public/MovieGraphImageSequenceOutputNode.h
@@ -197,7 +197,7 @@ public:
 
 	/** The file name format used for writing out the burn-in image (if not composited). */
 	UPROPERTY(EditAnywhere, BlueprintReadWrite, Category = "Burn In", meta = (EditCondition = "bOverride_BurnInFileNameFormat"))
-	FString BurnInFileNameFormat = TEXT("{sequence_name}.{layer_name}.{renderer_name}.{frame_number}");
+	FString BurnInFileNameFormat = TEXT("{sequence_name}.{layer_name}.BurnIn.{frame_number}");
 
 protected:
 
@@ -349,7 +349,7 @@ public:
 
 	/** The layer name used for this burn-in inside of multi-layer EXRs. */
 	UPROPERTY(EditAnywhere, BlueprintReadWrite, Category = "Burn In", meta = (DisplayAfter = "BurnInClass", EditCondition = "bOverride_BurnInLayerNameFormat"))
-	FString BurnInLayerNameFormat;
+	FString BurnInLayerNameFormat = TEXT("BurnIn");
 };
 
 /**
@@ -434,7 +434,7 @@ public:
 
 	/** The file name format used for writing out the burn-in image (if not composited). */
 	UPROPERTY(EditAnywhere, BlueprintReadWrite, Category = "Burn In", meta = (EditCondition = "bOverride_BurnInFileNameFormat"))
-	FString BurnInFileNameFormat = TEXT("{sequence_name}.{layer_name}.{renderer_name}.{frame_number}");
+	FString BurnInFileNameFormat = TEXT("{sequence_name}.{layer_name}.BurnIn.{frame_number}");
 };
 
 
@@ -520,7 +520,7 @@ public:
 
 	/** The file name format used for writing out the burn-in image (if not composited). */
 	UPROPERTY(EditAnywhere, BlueprintReadWrite, Category = "Burn In", meta = (EditCondition = "bOverride_BurnInFileNameFormat"))
-	FString BurnInFileNameFormat = TEXT("{sequence_name}.{layer_name}.{renderer_name}.{frame_number}");
+	FString BurnInFileNameFormat = TEXT("{sequence_name}.{layer_name}.BurnIn.{frame_number}");
 };
 
 /**
@@ -607,7 +607,7 @@ public:
 
 	/** The file name format used for writing out the burn-in image (if not composited). */
 	UPROPERTY(EditAnywhere, BlueprintReadWrite, Category = "Burn In", meta = (EditCondition = "bOverride_BurnInFileNameFormat"))
-	FString BurnInFileNameFormat = TEXT("{sequence_name}.{layer_name}.{renderer_name}.{frame_number}");
+	FString BurnInFileNameFormat = TEXT("{sequence_name}.{layer_name}.BurnIn.{frame_number}");
 };
 
 #undef UE_API
