MoviePipeline: Added an additional flag to the Deferred/PathTraced Renderer nodes in the graph to optionally exclude the beauty pass from being written to disk (useful if the user only wants to output PPM passes).

diff --git a/Engine/Plugins/MovieScene/MovieRenderPipeline/Source/MovieRenderPipelineCore/Private/Graph/Nodes/MovieGraphVideoOutputNode.cpp b/Engine/Plugins/MovieScene/MovieRenderPipeline/Source/MovieRenderPipelineCore/Private/Graph/Nodes/MovieGraphVideoOutputNode.cpp
--- a/Engine/Plugins/MovieScene/MovieRenderPipeline/Source/MovieRenderPipelineCore/Private/Graph/Nodes/MovieGraphVideoOutputNode.cpp
+++ b/Engine/Plugins/MovieScene/MovieRenderPipeline/Source/MovieRenderPipelineCore/Private/Graph/Nodes/MovieGraphVideoOutputNode.cpp
@@ -139,7 +139,7 @@ void UMovieGraphVideoOutputNode::OnReceiveImageDataImpl(UMovieGraphPipeline* InP
 		}
 		
 		// Write this frame's beauty pass, and pass along other passes so the encoder can composite them
-		if (RenderPassData.Key.SubResourceName == TEXT("beauty"))
+		if (Payload->bIsBeautyPass)
 		{
 			TArray<FMovieGraphPassData> CompositesForThisCamera;
 			for (FMovieGraphPassData& CompositePass : CompositedPasses)
diff --git a/Engine/Plugins/MovieScene/MovieRenderPipeline/Source/MovieRenderPipelineCore/Public/Graph/MovieGraphDataTypes.h b/Engine/Plugins/MovieScene/MovieRenderPipeline/Source/MovieRenderPipelineCore/Public/Graph/MovieGraphDataTypes.h
--- a/Engine/Plugins/MovieScene/MovieRenderPipeline/Source/MovieRenderPipelineCore/Public/Graph/MovieGraphDataTypes.h
+++ b/Engine/Plugins/MovieScene/MovieRenderPipeline/Source/MovieRenderPipelineCore/Public/Graph/MovieGraphDataTypes.h
@@ -406,6 +406,7 @@ namespace UE::MovieGraph
 			, bFetchFromAccumulator(false)
 			, bCompositeOnOtherRenders(false)
 			, bAllowsCompositing(true)
+			, bIsBeautyPass(false)
 			, OverscanFraction(0.f)
 			, CompositingSortOrder(0)
 			, RenderLayerIndex(0)
@@ -460,6 +461,9 @@ namespace UE::MovieGraph
 		/** Set this to false if this pass should not allow other passes to be composited onto it. */
 		bool bAllowsCompositing;
 
+		/** Whether this sample represents a sample in the main beauty pass. */
+		bool bIsBeautyPass;
+
 		/** When using high-res tiling, how many pixels does each tile overlap the adjacent tiles (on each side)? This should be zero if not using tiling. */
 		FIntPoint OverlappedPad;
 
diff --git a/Engine/Plugins/MovieScene/MovieRenderPipeline/Source/MovieRenderPipelineRenderPasses/Private/Graph/Nodes/MovieGraphDeferredPassNode.cpp b/Engine/Plugins/MovieScene/MovieRenderPipeline/Source/MovieRenderPipelineRenderPasses/Private/Graph/Nodes/MovieGraphDeferredPassNode.cpp
--- a/Engine/Plugins/MovieScene/MovieRenderPipeline/Source/MovieRenderPipelineRenderPasses/Private/Graph/Nodes/MovieGraphDeferredPassNode.cpp
+++ b/Engine/Plugins/MovieScene/MovieRenderPipeline/Source/MovieRenderPipelineRenderPasses/Private/Graph/Nodes/MovieGraphDeferredPassNode.cpp
@@ -19,6 +19,7 @@ UMovieGraphDeferredRenderPassNode::UMovieGraphDeferredRenderPassNode()
 	, bDisableToneCurve(false)
 	, bAllowOCIO(true)
 	, ViewModeIndex(VMI_Lit)
+	, bIncludeBeautyRenderInOutput(true)
 	, bEnableHighResolutionTiling(false)
 	, TileCount(1)
 	, OverlapPercentage(0.f)
@@ -90,6 +91,11 @@ EViewModeIndex UMovieGraphDeferredRenderPassNode::GetViewModeIndex() const
 	return ViewModeIndex;
 }
 
+bool UMovieGraphDeferredRenderPassNode::GetWriteBeautyPassToDisk() const
+{
+	return bIncludeBeautyRenderInOutput;
+}
+
 bool UMovieGraphDeferredRenderPassNode::GetWriteAllSamples() const
 {
 	return bWriteAllSamples;
diff --git a/Engine/Plugins/MovieScene/MovieRenderPipeline/Source/MovieRenderPipelineRenderPasses/Private/Graph/Nodes/MovieGraphPathTracerPassNode.cpp b/Engine/Plugins/MovieScene/MovieRenderPipeline/Source/MovieRenderPipelineRenderPasses/Private/Graph/Nodes/MovieGraphPathTracerPassNode.cpp
--- a/Engine/Plugins/MovieScene/MovieRenderPipeline/Source/MovieRenderPipelineRenderPasses/Private/Graph/Nodes/MovieGraphPathTracerPassNode.cpp
+++ b/Engine/Plugins/MovieScene/MovieRenderPipeline/Source/MovieRenderPipelineRenderPasses/Private/Graph/Nodes/MovieGraphPathTracerPassNode.cpp
@@ -26,6 +26,7 @@ UMovieGraphPathTracerRenderPassNode::UMovieGraphPathTracerRenderPassNode()
 	, bWriteAllSamples(false)
 	, bDisableToneCurve(false)
 	, bAllowOCIO(true)
+	, bIncludeBeautyRenderInOutput(true)
 	, TileCount(1)
 {
 	RendererName = TEXT("PathTraced");
@@ -236,6 +237,11 @@ bool UMovieGraphPathTracerRenderPassNode::GetAllowDenoiser() const
 	return bEnableDenoiser;
 }
 
+bool UMovieGraphPathTracerRenderPassNode::GetWriteBeautyPassToDisk() const
+{
+	return bIncludeBeautyRenderInOutput;
+}
+
 FEngineShowFlags UMovieGraphPathTracerRenderPassNode::GetShowFlags() const
 {
 	FEngineShowFlags OutShowFlags = Super::GetShowFlags();
diff --git a/Engine/Plugins/MovieScene/MovieRenderPipeline/Source/MovieRenderPipelineRenderPasses/Private/Graph/Renderers/MovieGraphDeferredPanoramicPass.cpp b/Engine/Plugins/MovieScene/MovieRenderPipeline/Source/MovieRenderPipelineRenderPasses/Private/Graph/Renderers/MovieGraphDeferredPanoramicPass.cpp
--- a/Engine/Plugins/MovieScene/MovieRenderPipeline/Source/MovieRenderPipelineRenderPasses/Private/Graph/Renderers/MovieGraphDeferredPanoramicPass.cpp
+++ b/Engine/Plugins/MovieScene/MovieRenderPipeline/Source/MovieRenderPipelineRenderPasses/Private/Graph/Renderers/MovieGraphDeferredPanoramicPass.cpp
@@ -684,6 +684,7 @@ void FMovieGraphDeferredPanoramicPass::Render(const FMovieGraphTraversalContext&
 					SampleState.bAllowOCIO = ParentNodeThisFrame->GetAllowOCIO();
 					SampleState.bForceLosslessCompression = ParentNodeThisFrame->GetForceLosslessCompression();
 					SampleState.bAllowsCompositing = ParentNodeThisFrame->GetAllowsCompositing();
+					SampleState.bIsBeautyPass = true;
 					SampleState.SceneCaptureSource = SceneCaptureSource;
 					SampleState.CompositingSortOrder = 10;
 					SampleState.RenderLayerIndex = LayerData.LayerIndex;
diff --git a/Engine/Plugins/MovieScene/MovieRenderPipeline/Source/MovieRenderPipelineRenderPasses/Private/Graph/Renderers/MovieGraphDeferredPass.cpp b/Engine/Plugins/MovieScene/MovieRenderPipeline/Source/MovieRenderPipelineRenderPasses/Private/Graph/Renderers/MovieGraphDeferredPass.cpp
--- a/Engine/Plugins/MovieScene/MovieRenderPipeline/Source/MovieRenderPipelineRenderPasses/Private/Graph/Renderers/MovieGraphDeferredPass.cpp
+++ b/Engine/Plugins/MovieScene/MovieRenderPipeline/Source/MovieRenderPipelineRenderPasses/Private/Graph/Renderers/MovieGraphDeferredPass.cpp
@@ -129,6 +129,9 @@ void FMovieGraphDeferredPass::Setup(TWeakObjectPtr<UMovieGraphDefaultRenderer> I
 	FramesToDelayPostSubmission = RenderPassNode ? RenderPassNode->GetCoolingDownFrameCount() : 0;
 	RemainingCooldownReadbackFrames = FramesToDelayPostSubmission;
 
+	// The user may request that the beauty pass isn't written to disk (sometimes used if only PPM passes are desired).
+	bWriteBeautyPassToDisk = ParentNode->GetWriteBeautyPassToDisk();
+
 	UE_LOG(LogMovieRenderPipeline, Log, TEXT("Initialized Renderer: %s Layer: %s TileCount[%d, %d]/[%d, %d] AutoExposurePass: %d HighRes: %d PageToSystem: %d TotalSceneViewStateReferences: %d"),
 		*RenderDataIdentifier.RendererName, *RenderDataIdentifier.LayerName, ParentNode->GetTileCount().X, ParentNode->GetTileCount().Y, TileCountWithHistory.X, TileCountWithHistory.Y,
 		bHasAutoExposurePass, ParentNode->GetEnableHighResolutionTiling(), ParentNode->GetEnablePageToSystemMemory(), NumSceneViewStateReferencesAllocated);
@@ -160,8 +163,11 @@ void FMovieGraphDeferredPass::GatherOutputPasses(UMovieGraphEvaluatedConfig* InC
 {
 	FMovieGraphImagePassBase::GatherOutputPasses(InConfig,OutExpectedPasses);
 	
-	// Add our pre-calculated identifier
-	OutExpectedPasses.Add(RenderDataIdentifier);
+	// Add our pre-calculated identifier if we're writing out the beauty pass.
+	if (bWriteBeautyPassToDisk)
+	{
+		OutExpectedPasses.Add(RenderDataIdentifier);
+	}
 
 	if (const UMovieGraphImagePassBaseNode* ParentNode = GetParentNode(InConfig))
 	{
@@ -467,6 +473,7 @@ void FMovieGraphDeferredPass::Render(const FMovieGraphTraversalContext& InFrameT
 					SampleState.bAllowOCIO = ParentNodeThisFrame->GetAllowOCIO();
 					SampleState.bForceLosslessCompression = ParentNodeThisFrame->GetForceLosslessCompression();
 					SampleState.bAllowsCompositing = ParentNodeThisFrame->GetAllowsCompositing();
+					SampleState.bIsBeautyPass = true;
 					SampleState.SceneCaptureSource = SceneCaptureSource;
 					SampleState.CompositingSortOrder = 10;
 					SampleState.RenderLayerIndex = LayerData.LayerIndex;
@@ -477,7 +484,7 @@ void FMovieGraphDeferredPass::Render(const FMovieGraphTraversalContext& InFrameT
 				MetadataHelper::UpdateSpatialSampleMetadata(SceneViewInitOptions, SampleState.AdditionalFileMetadata);
 
 				// If this was just to contribute to the history buffer, no need to go any further. Never discard if we're writing individual samples, though.
-				bool bDiscardOutput = (InTimeData.bDiscardOutput || ShouldDiscardOutput(ViewFamily, CameraInfo)) && !SampleState.bWriteSampleToDisk;
+				bool bRenderPass = (!InTimeData.bDiscardOutput && !ShouldDiscardOutput(ViewFamily, CameraInfo)) || SampleState.bWriteSampleToDisk;
 
 				if (UMovieGraphImagePassBaseNode* ParentNode = GetParentNode(InFrameTraversalContext.Time.EvaluatedConfig))
 				{
@@ -494,7 +501,7 @@ void FMovieGraphDeferredPass::Render(const FMovieGraphTraversalContext& InFrameT
 
 						// If we're not going to keep the output for the main pass then we don't want to create forwarding endpoints,
 						// as they'll read back data for the discarded main pass results and then try to pass them on.
-						if (bDiscardOutput)
+						if (!bRenderPass)
 						{
 							continue;
 						}
@@ -523,6 +530,9 @@ void FMovieGraphDeferredPass::Render(const FMovieGraphTraversalContext& InFrameT
 						// Some Post Process materials may need lossless compression in order to preserve the data properly.
 						PassSampleState.bForceLosslessCompression = PostProcessPass.bUseLosslessCompression;
 
+						// PPM passes are never the beauty pass.
+						PassSampleState.bIsBeautyPass = false;
+
 						// Give a lower priority to materials so they show up after the main pass in multi-layer exrs.
 						PassSampleState.CompositingSortOrder = SampleState.CompositingSortOrder + 1;
 						BufferPipe->AddEndpoint(MakeForwardingEndpoint(PassSampleState, InTimeData));
@@ -555,6 +565,7 @@ void FMovieGraphDeferredPass::Render(const FMovieGraphTraversalContext& InFrameT
 				}
 
 				// Data may not be something we actually want to read back and write to disk after rendering.
+				const bool bDiscardOutput = !bRenderPass || !bWriteBeautyPassToDisk;
 				if (bDiscardOutput)
 				{
 					continue;
diff --git a/Engine/Plugins/MovieScene/MovieRenderPipeline/Source/MovieRenderPipelineRenderPasses/Public/Graph/Nodes/MovieGraphDeferredPassNode.h b/Engine/Plugins/MovieScene/MovieRenderPipeline/Source/MovieRenderPipelineRenderPasses/Public/Graph/Nodes/MovieGraphDeferredPassNode.h
--- a/Engine/Plugins/MovieScene/MovieRenderPipeline/Source/MovieRenderPipelineRenderPasses/Public/Graph/Nodes/MovieGraphDeferredPassNode.h
+++ b/Engine/Plugins/MovieScene/MovieRenderPipeline/Source/MovieRenderPipelineRenderPasses/Public/Graph/Nodes/MovieGraphDeferredPassNode.h
@@ -45,6 +45,7 @@ protected:
 
 	// UMovieGraphImagePassBaseNode Interface
 	UE_API virtual EViewModeIndex GetViewModeIndex() const override;
+	UE_API virtual bool GetWriteBeautyPassToDisk() const override;
 	// ~UMovieGraphImagePassBaseNode Interface
 
 public:
@@ -66,6 +67,9 @@ public:
 	UPROPERTY(EditAnywhere, BlueprintReadWrite, Category = Overrides, meta = (InlineEditConditionToggle))
 	uint8 bOverride_bWriteAllSamples : 1;
 
+	UPROPERTY(EditAnywhere, BlueprintReadWrite, Category = Overrides, meta = (InlineEditConditionToggle))
+	uint8 bOverride_bIncludeBeautyRenderInOutput : 1;
+	
 	UPROPERTY(EditAnywhere, BlueprintReadWrite, Category = Overrides, meta = (InlineEditConditionToggle))
 	uint8 bOverride_AdditionalPostProcessMaterials : 1;
 
@@ -134,6 +138,14 @@ public:
 	UPROPERTY(EditAnywhere, BlueprintReadWrite, Category = "View Mode", meta=(EditCondition="bOverride_ViewModeIndex", InvalidEnumValues = "VMI_PathTracing,VMI_VisualizeBuffer,VMI_Unknown"))
 	TEnumAsByte<EViewModeIndex> ViewModeIndex;
 
+	/**
+	 * Whether the main beauty pass should be written to disk.
+	 *
+	 * If just Post Process Material passes should be written to disk, turn this off. 
+	 */
+	UPROPERTY(EditAnywhere, BlueprintReadWrite, Category = "Post Process Materials", meta=(EditCondition="bOverride_bIncludeBeautyRenderInOutput"))
+	bool bIncludeBeautyRenderInOutput;
+
 	/**
 	* An array of additional post-processing materials to run after the frame is rendered. Using this feature may add a notable amount of render time.
 	*/
diff --git a/Engine/Plugins/MovieScene/MovieRenderPipeline/Source/MovieRenderPipelineRenderPasses/Public/Graph/Nodes/MovieGraphImagePassBaseNode.h b/Engine/Plugins/MovieScene/MovieRenderPipeline/Source/MovieRenderPipelineRenderPasses/Public/Graph/Nodes/MovieGraphImagePassBaseNode.h
--- a/Engine/Plugins/MovieScene/MovieRenderPipeline/Source/MovieRenderPipelineRenderPasses/Public/Graph/Nodes/MovieGraphImagePassBaseNode.h
+++ b/Engine/Plugins/MovieScene/MovieRenderPipeline/Source/MovieRenderPipelineRenderPasses/Public/Graph/Nodes/MovieGraphImagePassBaseNode.h
@@ -76,6 +76,9 @@ public:
 	/** Whether this pass allows other passes to be composited on it. */
 	virtual bool GetAllowsCompositing() const { return true; }
 
+	/** Whether this node's beauty pass should be written to disk. For some nodes, this will have no effect because there is no "beauty" pass. */
+	virtual bool GetWriteBeautyPassToDisk() const { return true; }
+
 	/** Whether lossless compression must be forced on for this node. */
 	virtual bool GetForceLosslessCompression() const { return false; }
 	
diff --git a/Engine/Plugins/MovieScene/MovieRenderPipeline/Source/MovieRenderPipelineRenderPasses/Public/Graph/Nodes/MovieGraphPathTracerPassNode.h b/Engine/Plugins/MovieScene/MovieRenderPipeline/Source/MovieRenderPipelineRenderPasses/Public/Graph/Nodes/MovieGraphPathTracerPassNode.h
--- a/Engine/Plugins/MovieScene/MovieRenderPipeline/Source/MovieRenderPipelineRenderPasses/Public/Graph/Nodes/MovieGraphPathTracerPassNode.h
+++ b/Engine/Plugins/MovieScene/MovieRenderPipeline/Source/MovieRenderPipelineRenderPasses/Public/Graph/Nodes/MovieGraphPathTracerPassNode.h
@@ -56,14 +56,14 @@ public:
 	UE_API virtual bool GetDisableToneCurve() const override;
 	UE_API virtual bool GetAllowOCIO() const override;
 	UE_API virtual bool GetAllowDenoiser() const override;
-	UE_API virtual TUniquePtr<UE::MovieGraph::Rendering::FMovieGraphImagePassBase> CreateInstance() const;
+	UE_API virtual bool GetWriteBeautyPassToDisk() const override;
+	UE_API virtual TUniquePtr<UE::MovieGraph::Rendering::FMovieGraphImagePassBase> CreateInstance() const override;
 	UE_API virtual FEngineShowFlags GetShowFlags() const override;
 	UE_API virtual bool GetEnableHighResolutionTiling() const override;
 	UE_API virtual FIntPoint GetTileCount() const override;
 	UE_API virtual float GetTileOverlapPercentage() const override;
 	UE_API virtual bool GetEnablePageToSystemMemory() const override;
 	UE_API virtual bool GetEnableHistoryPerTile() const override;
-
 	// ~UMovieGraphImagePassBaseNode Interface
 
 	// UMovieGraphSettingNode Interface
@@ -131,6 +131,9 @@ public:
 	UPROPERTY(EditAnywhere, BlueprintReadWrite, Category = Overrides, meta = (InlineEditConditionToggle))
 	uint8 bOverride_bWriteAllSamples : 1;
 
+	UPROPERTY(EditAnywhere, BlueprintReadWrite, Category = Overrides, meta = (InlineEditConditionToggle))
+	uint8 bOverride_bIncludeBeautyRenderInOutput : 1;
+
 	UPROPERTY(EditAnywhere, BlueprintReadWrite, Category = Overrides, meta = (InlineEditConditionToggle))
 	uint8 bOverride_AdditionalPostProcessMaterials : 1;
 
@@ -244,6 +247,14 @@ public:
 	UPROPERTY(EditAnywhere, BlueprintReadWrite, Category = "Lighting Components", DisplayName = "Indirect Volume", meta = (EditCondition = "bOverride_bLightingComponents_IncludeIndirectVolume"))
 	bool bLightingComponents_IncludeIndirectVolume = true;
 
+	/**
+	 * Whether the main beauty pass should be written to disk.
+	 *
+	 * If just Post Process Material passes should be written to disk, turn this off. 
+	 */
+	UPROPERTY(EditAnywhere, BlueprintReadWrite, Category = "Post Process Materials", meta=(EditCondition="bOverride_bIncludeBeautyRenderInOutput"))
+	bool bIncludeBeautyRenderInOutput;
+
 	/**
 	* An array of additional post-processing materials to run after the frame is rendered. Using this feature may add a notable amount of render time.
 	*/
diff --git a/Engine/Plugins/MovieScene/MovieRenderPipeline/Source/MovieRenderPipelineRenderPasses/Public/Graph/Renderers/MovieGraphDeferredPass.h b/Engine/Plugins/MovieScene/MovieRenderPipeline/Source/MovieRenderPipelineRenderPasses/Public/Graph/Renderers/MovieGraphDeferredPass.h
--- a/Engine/Plugins/MovieScene/MovieRenderPipeline/Source/MovieRenderPipelineRenderPasses/Public/Graph/Renderers/MovieGraphDeferredPass.h
+++ b/Engine/Plugins/MovieScene/MovieRenderPipeline/Source/MovieRenderPipelineRenderPasses/Public/Graph/Renderers/MovieGraphDeferredPass.h
@@ -26,8 +26,7 @@ namespace UE::MovieGraph::Rendering
 		UE_API virtual void AddReferencedObjects(FReferenceCollector& Collector) override;
 		UE_API virtual FName GetBranchName() const override;
 		UE_API virtual UMovieGraphImagePassBaseNode* GetParentNode(UMovieGraphEvaluatedConfig* InConfig) const override;
-		UE_API bool ShouldDiscardOutput(const TSharedRef<FSceneViewFamilyContext>& InFamily, const UE::MovieGraph::DefaultRenderer::FCameraInfo& InCameraInfo) const override;
-
+		UE_API virtual bool ShouldDiscardOutput(const TSharedRef<FSceneViewFamilyContext>& InFamily, const UE::MovieGraph::DefaultRenderer::FCameraInfo& InCameraInfo) const override;
 		// End FMovieGraphImagePassBase
 			
 	protected:
@@ -62,6 +61,9 @@ namespace UE::MovieGraph::Rendering
 		// Did we initialize a auto-exposure sceneview history during setup?
 		bool bHasAutoExposurePass = false;
 
+		/** Whether the main beauty pass should be written to disk. Turn this off if only the PPM passes should be written. */
+		bool bWriteBeautyPassToDisk = true;
+
 		// Track the last Accumulator resolution we used, so that we can detect when it is changed and log that information.
 		FIntPoint PrevAccumulatorResolution;
 		// Track the last backbuffer resolution we used, so that we can detect when it is changed and log that information.
