MoviePipeline: Added the ability to use MRG (and CAT) {tokens} in more properties. This is done via the new ResolveTokenContainingProperties() method on settings nodes.

diff --git a/Engine/Plugins/MovieScene/MovieRenderPipeline/Source/MovieRenderPipelineCore/Private/Graph/MovieGraphBlueprintLibrary.cpp b/Engine/Plugins/MovieScene/MovieRenderPipeline/Source/MovieRenderPipelineCore/Private/Graph/MovieGraphBlueprintLibrary.cpp
--- a/Engine/Plugins/MovieScene/MovieRenderPipeline/Source/MovieRenderPipelineCore/Private/Graph/MovieGraphBlueprintLibrary.cpp
+++ b/Engine/Plugins/MovieScene/MovieRenderPipeline/Source/MovieRenderPipelineCore/Private/Graph/MovieGraphBlueprintLibrary.cpp
@@ -95,258 +95,277 @@ namespace
 			}
 		}
 	}
-}
 
-FFrameRate UMovieGraphBlueprintLibrary::GetEffectiveFrameRate(UMovieGraphGlobalOutputSettingNode* InNode, const FFrameRate& InDefaultRate)
-{
-	if (InNode && InNode->bOverride_OutputFrameRate)
+	FString ResolveFormatArgumentsHelper(const FString& InFormatString, const FMovieGraphFilenameResolveParams& InParams, const bool bResolveToPath, FMovieGraphResolveArgs& OutMergedFormatArgs)
 	{
-		return InNode->OutputFrameRate;
-	}
+		// There are a number of different sources for file/metadata KVPs. 
+		// UMovieGraphPipeline - these are handled as resolve params to allow resolving during edit time
+		// UMoviePipelineExecutorJob - author, sequence, level, etc.
+		// UMovieGraphSettingNode - every node can contribute it's own file/metadata KVP based on their own settings
 
-	return InDefaultRate;
-}
+		// Keep track of the resolve args that are generated JUST for this call to ResolveFilenameFormatArguments(). This generates a record of all tokens
+		// that were generated from MRG, and thus allows us to determine which tokens originate from outside MRG. These will be merged with
+		// OutMergedFormatArgs before the function returns.
+		FMovieGraphResolveArgs GeneratedResolveArgs;
+		
+		bool bOverwriteExisting = true;
+
+		// Data from the UMovieGraphPipeline comes from the resolve params themselves to allow resolving without a runtime instance.
+		{
+			// Zero-pad our frame numbers when we format the strings. Some programs struggle when ingesting frames that 
+			// go 1,2,3,...,10,11. To work around this issue we allow the user to specify how many zeros they want to
+			// pad the numbers with, 0001, 0002, etc. We also allow offsetting the output frame numbers, this is useful
+			// when your sequence starts at zero and you use handle frames (which would cause negative output frame 
+			// numbers), so we allow the user to add a fixed amount to all output to ensure they are positive.
+			int32 FrameNumberOffset = InParams.FrameNumberOffset;
+			FString FrameNumber = UE::MoviePipeline::GetPaddingFormatString(InParams.ZeroPadFrameNumberCount, InParams.RootFrameNumber + FrameNumberOffset); // Sequence Frame #
+			FString FrameNumberShot = UE::MoviePipeline::GetPaddingFormatString(InParams.ZeroPadFrameNumberCount, InParams.ShotFrameNumber + FrameNumberOffset); // Shot Frame #
+			FString FrameNumberRel = UE::MoviePipeline::GetPaddingFormatString(InParams.ZeroPadFrameNumberCount, InParams.RootFrameNumberRel + FrameNumberOffset); // Relative to 0
+			FString FrameNumberShotRel = UE::MoviePipeline::GetPaddingFormatString(InParams.ZeroPadFrameNumberCount, InParams.ShotFrameNumberRel + FrameNumberOffset); // Relative to 0 within the shot.
+
+			// If time dilation is in effect, RootFrameNumber and ShotFrameNumber will contain duplicates and the files will overwrite each other, 
+			// so we force them into relative mode and then warn users we did that (as their numbers will be out of sync with what they see in the level sequence now).
+			if (InFormatString.Contains(TEXT("{frame")) && InParams.bForceRelativeFrameNumbers && !InFormatString.Contains(TEXT("_rel}")))
+			{
+				UE_LOG(LogMovieRenderPipeline, Warning, TEXT("Time Dilation was used but output format does not use relative time, forcing relative numbers. Change {frame_number} to {frame_number_rel} (or the shot version) to remove this message."));
+			}
 
-FString UMovieGraphBlueprintLibrary::ResolveFilenameFormatArguments(const FString& InFormatString, const FMovieGraphFilenameResolveParams& InParams, FMovieGraphResolveArgs& OutMergedFormatArgs)
-{
-	// There are a number of different sources for file/metadata KVPs. 
-	// UMovieGraphPipeline - these are handled as resolve params to allow resolving during edit time
-	// UMoviePipelineExecutorJob - author, sequence, level, etc.
-	// UMovieGraphSettingNode - every node can contribute it's own file/metadata KVP based on their own settings
+			// Ensure they used relative frame numbers in the output so they get the right number of output frames.
+			if (InParams.bForceRelativeFrameNumbers)
+			{
+				FrameNumber = FrameNumberRel;
+				FrameNumberShot = FrameNumberShotRel;
+			}
 
-	// Keep track of the resolve args that are generated JUST for this call to ResolveFilenameFormatArguments(). This generates a record of all tokens
-	// that were generated from MRG, and thus allows us to determine which tokens originate from outside MRG. These will be merged with
-	// OutMergedFormatArgs before the function returns.
-	FMovieGraphResolveArgs GeneratedResolveArgs;
-	
-	bool bOverwriteExisting = true;
+			FString ShotName = InParams.Shot ? InParams.Shot->OuterName : FString();
+			ShotName = ShotName.Len() > 0 ? ShotName : TEXT("NoShot");
 
-	// Data from the UMovieGraphPipeline comes from the resolve params themselves to allow resolving without a runtime instance.
-	{
-		// Zero-pad our frame numbers when we format the strings. Some programs struggle when ingesting frames that 
-		// go 1,2,3,...,10,11. To work around this issue we allow the user to specify how many zeros they want to
-		// pad the numbers with, 0001, 0002, etc. We also allow offsetting the output frame numbers, this is useful
-		// when your sequence starts at zero and you use handle frames (which would cause negative output frame 
-		// numbers), so we allow the user to add a fixed amount to all output to ensure they are positive.
-		int32 FrameNumberOffset = InParams.FrameNumberOffset;
-		FString FrameNumber = UE::MoviePipeline::GetPaddingFormatString(InParams.ZeroPadFrameNumberCount, InParams.RootFrameNumber + FrameNumberOffset); // Sequence Frame #
-		FString FrameNumberShot = UE::MoviePipeline::GetPaddingFormatString(InParams.ZeroPadFrameNumberCount, InParams.ShotFrameNumber + FrameNumberOffset); // Shot Frame #
-		FString FrameNumberRel = UE::MoviePipeline::GetPaddingFormatString(InParams.ZeroPadFrameNumberCount, InParams.RootFrameNumberRel + FrameNumberOffset); // Relative to 0
-		FString FrameNumberShotRel = UE::MoviePipeline::GetPaddingFormatString(InParams.ZeroPadFrameNumberCount, InParams.ShotFrameNumberRel + FrameNumberOffset); // Relative to 0 within the shot.
+			// Use the legacy system to add filename and metadata kvp
+			MoviePipeline::GetOutputStateFormatArgs(GeneratedResolveArgs.FilenameArguments, GeneratedResolveArgs.FileMetadata, FrameNumber, FrameNumberShot, FrameNumberRel, FrameNumberShotRel, TEXT("DummyCameraToken"), ShotName);
 
-		// If time dilation is in effect, RootFrameNumber and ShotFrameNumber will contain duplicates and the files will overwrite each other, 
-		// so we force them into relative mode and then warn users we did that (as their numbers will be out of sync with what they see in the level sequence now).
-		if (InFormatString.Contains(TEXT("{frame")) && InParams.bForceRelativeFrameNumbers && !InFormatString.Contains(TEXT("_rel}")))
-		{
-			UE_LOG(LogMovieRenderPipeline, Warning, TEXT("Time Dilation was used but output format does not use relative time, forcing relative numbers. Change {frame_number} to {frame_number_rel} (or the shot version) to remove this message."));
-		}
+			// Add on Render Data Identifier, overwriting the dummy camera name above.
+			GeneratedResolveArgs.FilenameArguments.Add(TEXT("camera_name"), InParams.RenderDataIdentifier.CameraName);
+			GeneratedResolveArgs.FilenameArguments.Add(TEXT("renderer_name"), InParams.RenderDataIdentifier.RendererName);
+			GeneratedResolveArgs.FilenameArguments.Add(TEXT("renderer_sub_name"), InParams.RenderDataIdentifier.SubResourceName);
+			GeneratedResolveArgs.FilenameArguments.Add(TEXT("layer_name"), InParams.RenderDataIdentifier.LayerName);
+			GeneratedResolveArgs.FilenameArguments.Add(TEXT("branch_name"), InParams.RenderDataIdentifier.RootBranchName.ToString());
 
-		// Ensure they used relative frame numbers in the output so they get the right number of output frames.
-		if (InParams.bForceRelativeFrameNumbers)
-		{
-			FrameNumber = FrameNumberRel;
-			FrameNumberShot = FrameNumberShotRel;
+			// TODO: Some of these are per render layer and need to be stored that way. EXRs will have the metadata for all the layers/cameras/etc. in one file.
+			GeneratedResolveArgs.FileMetadata.Add(TEXT("unreal/cameraName"), InParams.RenderDataIdentifier.CameraName);
+			GeneratedResolveArgs.FileMetadata.Add(TEXT("unreal/rendererName"), InParams.RenderDataIdentifier.RendererName);
+			GeneratedResolveArgs.FileMetadata.Add(TEXT("unreal/subResourceName"), InParams.RenderDataIdentifier.SubResourceName);
+			GeneratedResolveArgs.FileMetadata.Add(TEXT("unreal/layerName"), InParams.RenderDataIdentifier.LayerName);
+			GeneratedResolveArgs.FileMetadata.Add(TEXT("unreal/branchName"), InParams.RenderDataIdentifier.RootBranchName.ToString());
 		}
 
-		FString ShotName = InParams.Shot ? InParams.Shot->OuterName : FString();
-		ShotName = ShotName.Len() > 0 ? ShotName : TEXT("NoShot");
 
-		// Use the legacy system to add filename and metadata kvp
-		MoviePipeline::GetOutputStateFormatArgs(GeneratedResolveArgs.FilenameArguments, GeneratedResolveArgs.FileMetadata, FrameNumber, FrameNumberShot, FrameNumberRel, FrameNumberShotRel, TEXT("DummyCameraToken"), ShotName);
+		// Add data from the job. ToDo: This probably needs to come from the job itself, so you can plug in different data sources eventually
+		{
+			// Frame Rate we handle separately, because it needs to be resolved.
+			double FrameRate = InParams.DefaultFrameRate.AsDecimal();
+			if (InParams.EvaluatedConfig)
+			{
+				const bool bIncludeCDOs = false;
+				UMovieGraphGlobalOutputSettingNode* OutputSettingNode = InParams.EvaluatedConfig->GetSettingForBranch<UMovieGraphGlobalOutputSettingNode>(UMovieGraphSettingNode::GlobalsPinName, bIncludeCDOs);
+
+				if (OutputSettingNode)
+				{
+					FrameRate = UMovieGraphBlueprintLibrary::GetEffectiveFrameRate(OutputSettingNode, InParams.DefaultFrameRate).AsDecimal();
+					bOverwriteExisting = OutputSettingNode->bOverwriteExistingOutput;
+				}
+			}
 
-		// Add on Render Data Identifier, overwriting the dummy camera name above.
-		GeneratedResolveArgs.FilenameArguments.Add(TEXT("camera_name"), InParams.RenderDataIdentifier.CameraName);
-		GeneratedResolveArgs.FilenameArguments.Add(TEXT("renderer_name"), InParams.RenderDataIdentifier.RendererName);
-		GeneratedResolveArgs.FilenameArguments.Add(TEXT("renderer_sub_name"), InParams.RenderDataIdentifier.SubResourceName);
-		GeneratedResolveArgs.FilenameArguments.Add(TEXT("layer_name"), InParams.RenderDataIdentifier.LayerName);
-		GeneratedResolveArgs.FilenameArguments.Add(TEXT("branch_name"), InParams.RenderDataIdentifier.RootBranchName.ToString());
+			GeneratedResolveArgs.FilenameArguments.Add(TEXT("frame_rate"), FString::SanitizeFloat(FrameRate));
+			GeneratedResolveArgs.FileMetadata.Add(TEXT("unreal/frameRate"), FString::SanitizeFloat(FrameRate));
 
-		// TODO: Some of these are per render layer and need to be stored that way. EXRs will have the metadata for all the layers/cameras/etc. in one file.
-		GeneratedResolveArgs.FileMetadata.Add(TEXT("unreal/cameraName"), InParams.RenderDataIdentifier.CameraName);
-		GeneratedResolveArgs.FileMetadata.Add(TEXT("unreal/rendererName"), InParams.RenderDataIdentifier.RendererName);
-		GeneratedResolveArgs.FileMetadata.Add(TEXT("unreal/subResourceName"), InParams.RenderDataIdentifier.SubResourceName);
-		GeneratedResolveArgs.FileMetadata.Add(TEXT("unreal/layerName"), InParams.RenderDataIdentifier.LayerName);
-		GeneratedResolveArgs.FileMetadata.Add(TEXT("unreal/branchName"), InParams.RenderDataIdentifier.RootBranchName.ToString());
-	}
+			FString LevelName;
+			FString SequenceName;
+			if (InParams.Job)
+			{
+				LevelName = InParams.Job->Map.GetAssetName();
+				SequenceName = InParams.Job->Sequence.GetAssetName();
+			}
+			
+			GeneratedResolveArgs.FilenameArguments.Add(TEXT("level_name"), LevelName);
+			GeneratedResolveArgs.FilenameArguments.Add(TEXT("sequence_name"), SequenceName);
+			
+			GeneratedResolveArgs.FileMetadata.Add(TEXT("unreal/levelName"), LevelName);
+			GeneratedResolveArgs.FileMetadata.Add(TEXT("unreal/sequenceName"), SequenceName);
+			
+			// Add KVP data from the job (date, time, job name, job author, job comment)
+			UE::MoviePipeline::GetSharedFormatArguments(GeneratedResolveArgs.FilenameArguments, GeneratedResolveArgs.FileMetadata, InParams.InitializationTime, InParams.Version, InParams.Job, InParams.InitializationTimeOffset);
+		}
+		
 
+		//  Now get the settings from our config. We need to gather KVP data from all possible nodes, even if not expressed in your configuration. This is because you might want to
+		// always use the {ts_count} token even if you don't have a Temporal Sample Count node to add it. We look at all possible class types and then get all instances of them in
+		// the evaluated graph. We look at all instances because "Named" nodes can end up having multiple copies in the final graph which may want to provide different data.
+		TArray<UClass*> AllSettingsNodeClasses = UE::MovieRenderPipeline::FindMoviePipelineSettingClasses(UMovieGraphSettingNode::StaticClass(), false);
 
-	// Add data from the job. ToDo: This probably needs to come from the job itself, so you can plug in different data sources eventually
-	{
-		// Frame Rate we handle separately, because it needs to be resolved.
-		double FrameRate = InParams.DefaultFrameRate.AsDecimal();
-		if (InParams.EvaluatedConfig)
+		// ToDo: This loops through class iterators every frame, we should probably initialize a copy of everything into the flattened config, since we could cache the classes
+		// once per run there. We don't cache the returned results here because you could potentially add/remove classes (via Blueprints) which would invalidate our cache.
+		for (UClass* InClass : AllSettingsNodeClasses)
 		{
-			const bool bIncludeCDOs = false;
-			UMovieGraphGlobalOutputSettingNode* OutputSettingNode = InParams.EvaluatedConfig->GetSettingForBranch<UMovieGraphGlobalOutputSettingNode>(UMovieGraphSettingNode::GlobalsPinName, bIncludeCDOs);
+			TArray<UMovieGraphSettingNode*> SettingInstances;
+			if (InParams.EvaluatedConfig)
+			{
+				const bool bIncludeCDOs = true;
+				const bool bExactMatch = true;
+				SettingInstances = InParams.EvaluatedConfig->GetSettingsForBranch(InClass, InParams.RenderDataIdentifier.RootBranchName, bIncludeCDOs, bExactMatch);
+			}
+			else
+			{
+				// GetFormatResolveArgs is const, but GetSettingsForBranch doesn't return a const pointer array,
+				// so we just get a non-const CDO pointer here instead, knowing that the function being called is const.
+				SettingInstances.Add(GetMutableDefault<UMovieGraphSettingNode>(InClass));
+			}
 
-			if (OutputSettingNode)
+			for (UMovieGraphSettingNode* Instance : SettingInstances)
 			{
-				FrameRate = GetEffectiveFrameRate(OutputSettingNode, InParams.DefaultFrameRate).AsDecimal();
-				bOverwriteExisting = OutputSettingNode->bOverwriteExistingOutput;
+				Instance->GetFormatResolveArgs(GeneratedResolveArgs, InParams.RenderDataIdentifier);
 			}
 		}
 
-		GeneratedResolveArgs.FilenameArguments.Add(TEXT("frame_rate"), FString::SanitizeFloat(FrameRate));
-		GeneratedResolveArgs.FileMetadata.Add(TEXT("unreal/frameRate"), FString::SanitizeFloat(FrameRate));
-
-		FString LevelName;
-		FString SequenceName;
-		if (InParams.Job)
+		// Copy the metadata in the parameters into the merged output version
+		for (const TPair<FString, FString>& FileMetadata : InParams.FileMetadata)
 		{
-			LevelName = InParams.Job->Map.GetAssetName();
-			SequenceName = InParams.Job->Sequence.GetAssetName();
+			GeneratedResolveArgs.FileMetadata.Add(FileMetadata.Key, FileMetadata.Value);
 		}
-		
-		GeneratedResolveArgs.FilenameArguments.Add(TEXT("level_name"), LevelName);
-		GeneratedResolveArgs.FilenameArguments.Add(TEXT("sequence_name"), SequenceName);
-		
-		GeneratedResolveArgs.FileMetadata.Add(TEXT("unreal/levelName"), LevelName);
-		GeneratedResolveArgs.FileMetadata.Add(TEXT("unreal/sequenceName"), SequenceName);
-		
-		// Add KVP data from the job (date, time, job name, job author, job comment)
-		UE::MoviePipeline::GetSharedFormatArguments(GeneratedResolveArgs.FilenameArguments, GeneratedResolveArgs.FileMetadata, InParams.InitializationTime, InParams.Version, InParams.Job, InParams.InitializationTimeOffset);
-	}
-	
 
-	//  Now get the settings from our config. We need to gather KVP data from all possible nodes, even if not expressed in your configuration. This is because you might want to
-	// always use the {ts_count} token even if you don't have a Temporal Sample Count node to add it. We look at all possible class types and then get all instances of them in
-	// the evaluated graph. We look at all instances because "Named" nodes can end up having multiple copies in the final graph which may want to provide different data.
-	TArray<UClass*> AllSettingsNodeClasses = UE::MovieRenderPipeline::FindMoviePipelineSettingClasses(UMovieGraphSettingNode::StaticClass(), false);
+		// We expect the incoming string to have a {file_dup} token where they want file duplication numbers to be handled
+		// but we don't actually want a value there (unless there's a collision). So we override it with an empty string by default.
+		GeneratedResolveArgs.FilenameArguments.Add(TEXT("file_dup"), FString());
 
-	// ToDo: This loops through class iterators every frame, we should probably initialize a copy of everything into the flattened config, since we could cache the classes
-	// once per run there. We don't cache the returned results here because you could potentially add/remove classes (via Blueprints) which would invalidate our cache.
-	for (UClass* InClass : AllSettingsNodeClasses)
-	{
-		TArray<UMovieGraphSettingNode*> SettingInstances;
-		if (InParams.EvaluatedConfig)
-		{
-			const bool bIncludeCDOs = true;
-			const bool bExactMatch = true;
-			SettingInstances = InParams.EvaluatedConfig->GetSettingsForBranch(InClass, InParams.RenderDataIdentifier.RootBranchName, bIncludeCDOs, bExactMatch);
-		}
-		else
+		// Format the {version} token so it's of the form "v###"
+		GeneratedResolveArgs.FilenameArguments.Add(TEXT("version"), FString::Printf(TEXT("v%0*d"), 3, InParams.Version));
+
+		// Overwrite any {tokens} with the user-supplied overrides if needed. This allows different requesters to share the same variables (ie: filename extension)
+		for (const TPair<FString, FString>& KVP : InParams.FileNameFormatOverrides)
 		{
-			// GetFormatResolveArgs is const, but GetSettingsForBranch doesn't return a const pointer array,
-			// so we just get a non-const CDO pointer here instead, knowing that the function being called is const.
-			SettingInstances.Add(GetMutableDefault<UMovieGraphSettingNode>(InClass));
+			GeneratedResolveArgs.FilenameArguments.Add(KVP.Key, KVP.Value);
 		}
 
-		for (UMovieGraphSettingNode* Instance : SettingInstances)
+		// Convert from our Python/BP exposed <FString, FString> to a named argument array for the formatter.
+		FStringFormatNamedArguments NamedArgs;
+		for (const TPair<FString, FString>& Argument : GeneratedResolveArgs.FilenameArguments)
 		{
-			Instance->GetFormatResolveArgs(GeneratedResolveArgs, InParams.RenderDataIdentifier);
+			NamedArgs.Add(Argument.Key, Argument.Value);
 		}
-	}
-
-	// Copy the metadata in the parameters into the merged output version
-	for (const TPair<FString, FString>& FileMetadata : InParams.FileMetadata)
-	{
-		GeneratedResolveArgs.FileMetadata.Add(FileMetadata.Key, FileMetadata.Value);
-	}
-
-	// We expect the incoming string to have a {file_dup} token where they want file duplication numbers to be handled
-	// but we don't actually want a value there (unless there's a collision). So we override it with an empty string by default.
-	GeneratedResolveArgs.FilenameArguments.Add(TEXT("file_dup"), FString());
-
-	// Format the {version} token so it's of the form "v###"
-	GeneratedResolveArgs.FilenameArguments.Add(TEXT("version"), FString::Printf(TEXT("v%0*d"), 3, InParams.Version));
-
-	// Overwrite any {tokens} with the user-supplied overrides if needed. This allows different requesters to share the same variables (ie: filename extension)
-	for (const TPair<FString, FString>& KVP : InParams.FileNameFormatOverrides)
-	{
-		GeneratedResolveArgs.FilenameArguments.Add(KVP.Key, KVP.Value);
-	}
 
-	// Convert from our Python/BP exposed <FString, FString> to a named argument array for the formatter.
-	FStringFormatNamedArguments NamedArgs;
-	for (const TPair<FString, FString>& Argument : GeneratedResolveArgs.FilenameArguments)
-	{
-		NamedArgs.Add(Argument.Key, Argument.Value);
-	}
+		// The format string may need to be modified from this point onwards, so make a copy so the source string isn't unintentionally changed.
+		FString FormatStringCopy = InFormatString;
 
-	// The format string may need to be modified from this point onwards, so make a copy so the source string isn't unintentionally changed.
-	FString FormatStringCopy = InFormatString;
+		// Resolve non-MRG tokens via the Naming Tokens system.
+		UpdateTokenMapsWithResolvedExternalTokens(FormatStringCopy, GeneratedResolveArgs.FilenameArguments, NamedArgs);
 
-	// Resolve non-MRG tokens via the Naming Tokens system.
-	UpdateTokenMapsWithResolvedExternalTokens(FormatStringCopy, GeneratedResolveArgs.FilenameArguments, NamedArgs);
+		// Now that all format args have been discovered, add them to the merged format args.
+		OutMergedFormatArgs.FileMetadata.Append(GeneratedResolveArgs.FileMetadata);
+		OutMergedFormatArgs.FilenameArguments.Append(GeneratedResolveArgs.FilenameArguments);
 
-	// Now that all format args have been discovered, add them to the merged format args.
-	OutMergedFormatArgs.FileMetadata.Append(GeneratedResolveArgs.FileMetadata);
-	OutMergedFormatArgs.FilenameArguments.Append(GeneratedResolveArgs.FilenameArguments);
-
-	// If no format string is provided, there's nothing left to do.
-	if (InFormatString.IsEmpty())
-	{
-		return FString();
-	}
-
-	// Apply all of our named args to the file and generate a path.
-	FString BaseFilename = FString::Format(*FormatStringCopy, NamedArgs);
-	if (InParams.bEnsureAbsolutePath)
-	{
-		BaseFilename = FPaths::ConvertRelativePathToFull(BaseFilename);
-	}
+		// If no format string is provided, there's nothing left to do.
+		if (InFormatString.IsEmpty())
+		{
+			return FString();
+		}
 
-	// Fix-up slashes
-	{
-		FPaths::NormalizeFilename(BaseFilename);
+		// Apply all of our named args to create a formatted string. This may or may not be a file path.
+		FString FormattedString = FString::Format(*FormatStringCopy, NamedArgs);
 
-		// Remove duplicate slashes. Recognize a double-slash at the beginning to ensure network paths remain intact. Note that this also works on
-		// network paths that start with "\\" because NormalizeFilename() will convert "\\" to "//".
-		if (BaseFilename.StartsWith(TEXT("//")))
+		// If not generating a file path, we can stop here.
+		if (!bResolveToPath)
 		{
-			BaseFilename.RightChopInline(1, EAllowShrinking::No);
-			FPaths::RemoveDuplicateSlashes(BaseFilename);
-			BaseFilename = TEXT("/") + MoveTemp(BaseFilename);
+			return FormattedString;
 		}
-		else
+		
+		if (InParams.bEnsureAbsolutePath)
 		{
-			FPaths::RemoveDuplicateSlashes(BaseFilename);
+			FormattedString = FPaths::ConvertRelativePathToFull(FormattedString);
 		}
-	}
 
-	// In the event of multiple dots in the filename, we need to replace them with a single dot to ensure a valid file name
-	while (BaseFilename.Contains(".."))
-	{
-		BaseFilename.ReplaceInline(TEXT(".."), TEXT("."));
-	}
-
-	// If we end with a "." character, remove it. The extension will put it back on. We can end up with this sometimes
-	// after resolving file format strings, ie: {sequence_name}.{frame_number} becomes {sequence_name}. for
-	// videos (which can't use frame_numbers).
-	BaseFilename.RemoveFromEnd(TEXT("."));
-	
-	// If the extension is not resolved, ThisTry will be a path
-	const FString ExtToken = TEXT(".{ext}");
-	const FString Extension = FString::Format(*ExtToken, NamedArgs);
-	const FString ResolvedExtension = (Extension == ExtToken) ? FString() : Extension;
-	FString ThisTry = BaseFilename + Extension;
-	
-	// Check that it's a valid path that we can write to.
-	if (UE::MoviePipeline::CanWriteToFile(*ThisTry, bOverwriteExisting))
-	{
-		return ThisTry;
-	}
+		// Fix-up slashes
+		{
+			FPaths::NormalizeFilename(FormattedString);
 
-	// If we got here it means that there was already a file there with that name and bOverwriteExisting was false.
-	// So we start by swapping in _(2) where the file_dup token is, and try again. If that fails, _(3), etc. We start
-	// at _(2) because the file they're trying to write (foo.png) is technically "_(1)", so we start the rename process
-	// at _(2) so it accurately reflects that this is the second version of that file.
-	int32 DuplicateIndex = 2;
-	while (true)
-	{
-		NamedArgs.Add(TEXT("file_dup"), FString::Printf(TEXT("_(%d)"), DuplicateIndex));
+			// Remove duplicate slashes. Recognize a double-slash at the beginning to ensure network paths remain intact. Note that this also works on
+			// network paths that start with "\\" because NormalizeFilename() will convert "\\" to "//".
+			if (FormattedString.StartsWith(TEXT("//")))
+			{
+				FormattedString.RightChopInline(1, EAllowShrinking::No);
+				FPaths::RemoveDuplicateSlashes(FormattedString);
+				FormattedString = TEXT("/") + MoveTemp(FormattedString);
+			}
+			else
+			{
+				FPaths::RemoveDuplicateSlashes(FormattedString);
+			}
+		}
 
-		// Re-resolve the format string now that we've reassigned frame_dup to a number.
-		ThisTry = FString::Format(*FormatStringCopy, NamedArgs) + ResolvedExtension;
+		// In the event of multiple dots in the filename, we need to replace them with a single dot to ensure a valid file name
+		while (FormattedString.Contains(".."))
+		{
+			FormattedString.ReplaceInline(TEXT(".."), TEXT("."));
+		}
 
-		// If the file doesn't exist, we can use that, else, increment the index and try again
+		// If we end with a "." character, remove it. The extension will put it back on. We can end up with this sometimes
+		// after resolving file format strings, ie: {sequence_name}.{frame_number} becomes {sequence_name}. for
+		// videos (which can't use frame_numbers).
+		FormattedString.RemoveFromEnd(TEXT("."));
+		
+		// If the extension is not resolved, ThisTry will be a path
+		const FString ExtToken = TEXT(".{ext}");
+		const FString Extension = FString::Format(*ExtToken, NamedArgs);
+		const FString ResolvedExtension = (Extension == ExtToken) ? FString() : Extension;
+		FString ThisTry = FormattedString + Extension;
+		
+		// Check that it's a valid path that we can write to.
 		if (UE::MoviePipeline::CanWriteToFile(*ThisTry, bOverwriteExisting))
 		{
 			return ThisTry;
 		}
 
-		++DuplicateIndex;
+		// If we got here it means that there was already a file there with that name and bOverwriteExisting was false.
+		// So we start by swapping in _(2) where the file_dup token is, and try again. If that fails, _(3), etc. We start
+		// at _(2) because the file they're trying to write (foo.png) is technically "_(1)", so we start the rename process
+		// at _(2) so it accurately reflects that this is the second version of that file.
+		int32 DuplicateIndex = 2;
+		while (true)
+		{
+			NamedArgs.Add(TEXT("file_dup"), FString::Printf(TEXT("_(%d)"), DuplicateIndex));
+
+			// Re-resolve the format string now that we've reassigned frame_dup to a number.
+			ThisTry = FString::Format(*FormatStringCopy, NamedArgs) + ResolvedExtension;
+
+			// If the file doesn't exist, we can use that, else, increment the index and try again
+			if (UE::MoviePipeline::CanWriteToFile(*ThisTry, bOverwriteExisting))
+			{
+				return ThisTry;
+			}
+
+			++DuplicateIndex;
+		}
+	}
+}
+
+FString UMovieGraphBlueprintLibrary::ResolveFilenameFormatArguments(const FString& InFormatString, const FMovieGraphFilenameResolveParams& InParams, FMovieGraphResolveArgs& OutMergedFormatArgs)
+{
+	constexpr bool bResolveToPath = true;
+	return ResolveFormatArgumentsHelper(InFormatString, InParams, bResolveToPath, OutMergedFormatArgs);
+}
+
+FString UMovieGraphBlueprintLibrary::ResolveFormatArguments(const FString& InFormatString, const FMovieGraphFilenameResolveParams& InParams, FMovieGraphResolveArgs& OutMergedFormatArgs)
+{
+	constexpr bool bResolveToPath = false;
+	return ResolveFormatArgumentsHelper(InFormatString, InParams, bResolveToPath, OutMergedFormatArgs);
+}
+
+FFrameRate UMovieGraphBlueprintLibrary::GetEffectiveFrameRate(UMovieGraphGlobalOutputSettingNode* InNode, const FFrameRate& InDefaultRate)
+{
+	if (InNode && InNode->bOverride_OutputFrameRate)
+	{
+		return InNode->OutputFrameRate;
 	}
+
+	return InDefaultRate;
 }
 
 int32 UMovieGraphBlueprintLibrary::ResolveVersionNumber(FMovieGraphFilenameResolveParams InParams, const bool bGetNextVersion)
diff --git a/Engine/Plugins/MovieScene/MovieRenderPipeline/Source/MovieRenderPipelineCore/Private/Graph/MovieGraphConfig.cpp b/Engine/Plugins/MovieScene/MovieRenderPipeline/Source/MovieRenderPipelineCore/Private/Graph/MovieGraphConfig.cpp
--- a/Engine/Plugins/MovieScene/MovieRenderPipeline/Source/MovieRenderPipelineCore/Private/Graph/MovieGraphConfig.cpp
+++ b/Engine/Plugins/MovieScene/MovieRenderPipeline/Source/MovieRenderPipelineCore/Private/Graph/MovieGraphConfig.cpp
@@ -2253,6 +2253,54 @@ UMovieGraphEvaluatedConfig* UMovieGraphConfig::CreateFlattenedGraph(const FMovie
 		}
 	}
 
+	// For properties that contain {tokens}, resolve their values.
+	{
+		// The identifier needs to be updated for each branch that the tokens are resolved on. Some properties in the identifier cannot be resolved
+		// here (like {camera_name}) because there's not enough context at this point.
+		FMovieGraphRenderDataIdentifier Identifier = InContext.RenderDataIdentifier;
+
+		// This lambda, in combination with the resolve params, are used to resolve the format strings. The resolve args aren't used here, but are
+		// required (they're reset because ResolveFormatArguments() appends to them, and there's no reason to continuously allocate a larger and
+		// larger array).
+		FMovieGraphFilenameResolveParams ResolveParams =
+			FMovieGraphFilenameResolveParams::MakeResolveParams(Identifier, nullptr, InContext.Time.EvaluatedConfig, InContext, {});
+		FMovieGraphResolveArgs ResolveArgs;
+		TFunction<void(FString&)> ResolveToken = [&ResolveParams, &ResolveArgs, &Identifier](FString& InValue) -> void
+		{
+			ResolveArgs.FileMetadata.Reset();
+			ResolveArgs.FilenameArguments.Reset();
+
+			// The identifier is updated for each branch/layer
+			ResolveParams.RenderDataIdentifier = Identifier;
+
+			// Assign the resolved value back to the property that's being resolved
+			InValue = UMovieGraphBlueprintLibrary::ResolveFormatArguments(InValue, ResolveParams, ResolveArgs);
+		};
+
+		// Give each setting node a chance to resolve its properties
+		FMovieGraphTokenResolveContext AdditionalContext;
+		for (const TPair<FName, FMovieGraphEvaluatedBranchConfig>& Pair : NewContext->BranchConfigMapping)
+		{
+			for (TObjectPtr<UMovieGraphNode>& Node : Pair.Value.GetNodes())
+			{
+				if (UMovieGraphSettingNode* Setting = Cast<UMovieGraphSettingNode>(Node))
+				{
+					UMovieGraphRenderLayerNode* RenderLayerNode =
+						Cast<UMovieGraphRenderLayerNode>(Pair.Value.GetNodeByClassExactMatch(UMovieGraphRenderLayerNode::StaticClass(), FString()));
+
+					// Provide additional context for the resolve; some of the identifier's properties can't be populated here, though. These updates
+					// will be picked up by the ResolveToken() lambda above.
+					Identifier.RootBranchName = Pair.Key;
+					Identifier.LayerName = RenderLayerNode ? RenderLayerNode->LayerName : TEXT("Unknown");
+
+					AdditionalContext.RenderDataIdentifier = &Identifier;
+					AdditionalContext.TraversalContext = &InContext;
+					Setting->ResolveTokenContainingProperties(ResolveToken, AdditionalContext);
+				}
+			}
+		}
+	}
+
 	// Check all nodes for validation errors. If any, these will be logged, and the render will be halted.
 	TArray<FText> ValidationErrors;
 	for (const TPair<FName, FMovieGraphEvaluatedBranchConfig>& Pair : NewContext->BranchConfigMapping)
diff --git a/Engine/Plugins/MovieScene/MovieRenderPipeline/Source/MovieRenderPipelineCore/Private/Graph/MovieGraphFilenameResolveParams.cpp b/Engine/Plugins/MovieScene/MovieRenderPipeline/Source/MovieRenderPipelineCore/Private/Graph/MovieGraphFilenameResolveParams.cpp
--- a/Engine/Plugins/MovieScene/MovieRenderPipeline/Source/MovieRenderPipelineCore/Private/Graph/MovieGraphFilenameResolveParams.cpp
+++ b/Engine/Plugins/MovieScene/MovieRenderPipeline/Source/MovieRenderPipelineCore/Private/Graph/MovieGraphFilenameResolveParams.cpp
@@ -19,7 +19,11 @@ FMovieGraphFilenameResolveParams FMovieGraphFilenameResolveParams::MakeResolvePa
 	
 	int32 RootFrameNumberRelOffset = 0;
 	int32 ShotFrameNumberRelOffset = 0;
-	if (ensureAlwaysMsgf(InPipeline, TEXT("InPipeline is not valid - ResolveParams will be created, but will be missing critical information.")))
+
+	// Initialize the job/shot with the pipeline if available; otherwise, fall back to the traversal context.
+	// Using the traversal context may result in a slightly under-initialized set of resolve parameters, but generally
+	// these parameters aren't needed in the contexts that a pipeline isn't available.
+	if (IsValid(InPipeline))
 	{
 		Params.InitializationTime = InPipeline->GetInitializationTime();
 		Params.InitializationTimeOffset = InPipeline->GetInitializationTimeOffset();
@@ -27,20 +31,29 @@ FMovieGraphFilenameResolveParams FMovieGraphFilenameResolveParams::MakeResolvePa
 		
 		if (InPipeline->GetActiveShotList().IsValidIndex(InTraversalContext.ShotIndex))
 		{
-			const TObjectPtr<UMoviePipelineExecutorShot>& Shot = InPipeline->GetActiveShotList()[InTraversalContext.ShotIndex];
-			Params.Version = Shot->ShotInfo.VersionNumber;
-			Params.Shot = Shot;
-			
-			// Calculate the offset (in frames) that relative frame numbers need to add to themselves so that they are correctly
-			// offset by the starting frame of the root sequence/shot (to match new updated relative behavior).
-			// We use the first shot to calculate the root version of the relative frame number offset, because that'll have the correct
-			// root time for rendered frames (that should account for handle frames (included) vs warmup (not)). It's safe to assume index
-			// 0 exists here because we don't null check ShotIndex and it's either going to be >= 0.
-			const TObjectPtr<UMoviePipelineExecutorShot>& RootShot = InPipeline->GetActiveShotList()[0];
-			
-			RootFrameNumberRelOffset = FFrameRate::TransformTime(RootShot->ShotInfo.InitialTimeInRoot, RootShot->ShotInfo.CachedTickResolution, RootShot->ShotInfo.CachedFrameRate).FloorToFrame().Value;
-			ShotFrameNumberRelOffset = FFrameRate::TransformTime(Shot->ShotInfo.InitialTimeInShot, Shot->ShotInfo.CachedTickResolution, Shot->ShotInfo.CachedFrameRate).FloorToFrame().Value;
+			Params.Shot = InPipeline->GetActiveShotList()[InTraversalContext.ShotIndex];
 		}
+
+		// Calculate the offset (in frames) that relative frame numbers need to add to themselves so that they are correctly
+		// offset by the starting frame of the root sequence/shot (to match new updated relative behavior).
+		// We use the first shot to calculate the root version of the relative frame number offset, because that'll have the correct
+		// root time for rendered frames (that should account for handle frames (included) vs warmup (not)). It's safe to assume index
+		// 0 exists here because we don't null check ShotIndex and it's either going to be >= 0.
+		const TObjectPtr<UMoviePipelineExecutorShot>& RootShot = InPipeline->GetActiveShotList()[0];
+
+		RootFrameNumberRelOffset = FFrameRate::TransformTime(RootShot->ShotInfo.InitialTimeInRoot, RootShot->ShotInfo.CachedTickResolution, RootShot->ShotInfo.CachedFrameRate).FloorToFrame().Value;
+	}
+	else
+	{
+		Params.Job = InTraversalContext.Job;
+		Params.Shot = InTraversalContext.Shot;
+	}
+
+	if (const TObjectPtr<UMoviePipelineExecutorShot>& Shot = Params.Shot)
+	{
+		Params.Version = Shot->ShotInfo.VersionNumber;
+		
+		ShotFrameNumberRelOffset = FFrameRate::TransformTime(Shot->ShotInfo.InitialTimeInShot, Shot->ShotInfo.CachedTickResolution, Shot->ShotInfo.CachedFrameRate).FloorToFrame().Value;
 	}
         
 	Params.RenderDataIdentifier = InRenderId;
diff --git a/Engine/Plugins/MovieScene/MovieRenderPipeline/Source/MovieRenderPipelineCore/Private/Graph/MovieGraphNode.cpp b/Engine/Plugins/MovieScene/MovieRenderPipeline/Source/MovieRenderPipelineCore/Private/Graph/MovieGraphNode.cpp
--- a/Engine/Plugins/MovieScene/MovieRenderPipeline/Source/MovieRenderPipelineCore/Private/Graph/MovieGraphNode.cpp
+++ b/Engine/Plugins/MovieScene/MovieRenderPipeline/Source/MovieRenderPipelineCore/Private/Graph/MovieGraphNode.cpp
@@ -658,6 +658,14 @@ TArray<FMovieGraphPinProperties> UMovieGraphSettingNode::GetOutputPinProperties(
 	return Properties;
 }
 
+void UMovieGraphSettingNode::ResolveTokenContainingProperties(TFunction<void(FString&)>& ResolveFunc, const FMovieGraphTokenResolveContext& InContext)
+{
+	// The default node has no properties that contain {tokens} that need to be resolved.
+
+	// Dev note: Ideally we could use UPROPERTY metadata to define these properties. However, UPROPERTY metadata isn't available in runtime
+	// (non-editor) builds, so it's not a viable solution.
+}
+
 void IMovieGraphEvaluationNodeInjector::InjectNodesPostEvaluation(const FName& InBranchName, UMovieGraphEvaluatedConfig* InEvaluatedConfig, TArray<UMovieGraphSettingNode*>& OutInjectedNodes)
 {
 	
diff --git a/Engine/Plugins/MovieScene/MovieRenderPipeline/Source/MovieRenderPipelineCore/Public/Graph/MovieGraphBlueprintLibrary.h b/Engine/Plugins/MovieScene/MovieRenderPipeline/Source/MovieRenderPipelineCore/Public/Graph/MovieGraphBlueprintLibrary.h
--- a/Engine/Plugins/MovieScene/MovieRenderPipeline/Source/MovieRenderPipelineCore/Public/Graph/MovieGraphBlueprintLibrary.h
+++ b/Engine/Plugins/MovieScene/MovieRenderPipeline/Source/MovieRenderPipelineCore/Public/Graph/MovieGraphBlueprintLibrary.h
@@ -34,18 +34,25 @@ public:
 
 	/**
 	* Takes a Movie Graph format string (in the form of {token}), a list of parameters (which normally come from the running UMovieGraphPipeline) and
-	* resolves them into a string. Unknown tokens are ignored. Which tokens can be resolved depends on the contents of InParams, tokens from settings
+	* resolves them into a file path. Unknown tokens are ignored. Which tokens can be resolved depends on the contents of InParams, tokens from settings
 	* rely on a evaluated config being provided, etc.
 	* @param	InFormatString		- Format string to attempt to resolve. Leave blank if just the format args should be populated.
 	* @param	InParams			- A list of parameters to use as source data for the resolve step. Normally comes from the UMovieGraphPipeline instance,
 	*								- but takes (mostly) POD here to allow using this function outside of the render runtime.
 	* @param	OutMergedFormatArgs - The set of KVP for both filename formats and file metadata that is generated as a result of this. Provided in case you
 	* 								- needed to do your own string resolving with the final dataset.
-	* @return						- The resolved format string. Returns an empty string if InFormatString is blank.
+	* @return						- The resolved file path string. Returns an empty string if InFormatString is blank.
 	*/
 	UFUNCTION(BlueprintCallable, Category = "Movie Graph")
 	static UE_API FString ResolveFilenameFormatArguments(const FString& InFormatString, const FMovieGraphFilenameResolveParams& InParams, FMovieGraphResolveArgs& OutMergedFormatArgs);
 
+	/**
+	 * Functions identically to ResolveFilenameFormatArguments(), with the exception that this variation does not resolve the format string to a
+	 * file path (any non-path format string can be provided).
+	 */
+	UFUNCTION(BlueprintCallable, Category = "Movie Graph")
+	static UE_API FString ResolveFormatArguments(const FString& InFormatString, const FMovieGraphFilenameResolveParams& InParams, FMovieGraphResolveArgs& OutMergedFormatArgs);
+
 	/**
 	 * If the version number is explicitly specified on the Output Setting node, returns that. Otherwise searches the
 	 * output directory for the highest version that already exists (and increments it by one if bGetNextVersion is
diff --git a/Engine/Plugins/MovieScene/MovieRenderPipeline/Source/MovieRenderPipelineCore/Public/Graph/MovieGraphConfig.h b/Engine/Plugins/MovieScene/MovieRenderPipeline/Source/MovieRenderPipelineCore/Public/Graph/MovieGraphConfig.h
--- a/Engine/Plugins/MovieScene/MovieRenderPipeline/Source/MovieRenderPipelineCore/Public/Graph/MovieGraphConfig.h
+++ b/Engine/Plugins/MovieScene/MovieRenderPipeline/Source/MovieRenderPipelineCore/Public/Graph/MovieGraphConfig.h
@@ -312,7 +312,7 @@ struct FMovieGraphEvaluatedBranchConfig
 {
 	GENERATED_BODY()
 
-	UMovieGraphNode* GetNodeByClassExactMatch(const TSubclassOf<UMovieGraphNode>& InClass, const FString& InName)
+	UMovieGraphNode* GetNodeByClassExactMatch(const TSubclassOf<UMovieGraphNode>& InClass, const FString& InName) const
 	{
 		if (const FMovieGraphEvaluatedSettingsStack* FoundStack = NamedNodes.Find(InName))
 		{
diff --git a/Engine/Plugins/MovieScene/MovieRenderPipeline/Source/MovieRenderPipelineCore/Public/Graph/MovieGraphNode.h b/Engine/Plugins/MovieScene/MovieRenderPipeline/Source/MovieRenderPipelineCore/Public/Graph/MovieGraphNode.h
--- a/Engine/Plugins/MovieScene/MovieRenderPipeline/Source/MovieRenderPipelineCore/Public/Graph/MovieGraphNode.h
+++ b/Engine/Plugins/MovieScene/MovieRenderPipeline/Source/MovieRenderPipelineCore/Public/Graph/MovieGraphNode.h
@@ -97,6 +97,21 @@ struct FMovieGraphPropertyInfo
 	}
 };
 
+/**
+ * Additional context provided to ResolveTokenContainingProperties(). May be helpful if the node needs to implement custom {token} resolve
+ * behavior that goes beyond the default resolve behavior that MRG provides.
+ */
+USTRUCT(BlueprintType)
+struct FMovieGraphTokenResolveContext
+{
+	GENERATED_BODY()
+	
+	FMovieGraphTokenResolveContext() = default;
+	
+	const FMovieGraphRenderDataIdentifier* RenderDataIdentifier = nullptr;
+	const FMovieGraphTraversalContext* TraversalContext = nullptr;
+};
+
 /** Describes a restriction on what kind of branch a node can be created in within the graph. */
 UENUM(BlueprintType)
 enum class EMovieGraphBranchRestriction : uint8
@@ -503,6 +518,20 @@ public:
 	* resolving the settings of the graph, the node only needs to read its own values.
 	*/
 	virtual void GetFormatResolveArgs(FMovieGraphResolveArgs& OutMergedFormatArgs, const FMovieGraphRenderDataIdentifier& InRenderDataIdentifier) const {}
+
+	/**
+	 * Resolves the values of properties that contain {tokens} to their final value. Graph evaluation performs this step after after the graph has
+	 * been completely flattened, and it's the node's responsibility to implement this method if any of its properties have {tokens} that need to
+	 * be resolved.
+	 *
+	 * All properties that need to be resolved should have ResolveFunc() called on them; this will update the node's property with the resolved
+	 * value. Regular MRG tokens, as well as tokens from the Naming Tokens (CAT) system, will be resolved. Nodes can also implement custom resolver
+	 * behavior here if needed, using the additional context provided. Select MRG tokens may not be resolved correctly here due to the point in the
+	 * graph evaluation lifecycle that this resolve step occurs (eg, {camera_name}).
+	 *
+	 * Nodes that implement this method should also take care to call the super class in order to ensure inherited properties are covered.
+	 */
+	MOVIERENDERPIPELINECORE_API virtual void ResolveTokenContainingProperties(TFunction<void(FString&)>& ResolveFunc, const FMovieGraphTokenResolveContext& InContext);
 	
 	/** Modify the Unreal URL and command line arguments when the node will be run in a new process. Only applies to nodes in the Globals branch. */
 	virtual void BuildNewProcessCommandLineArgsImpl(TArray<FString>& InOutUnrealURLParams, TArray<FString>& InOutCommandLineArgs, TArray<FString>& InOutDeviceProfileCvars, TArray<FString>& InOutExecCmds) const { }
diff --git a/Engine/Plugins/MovieScene/MovieRenderPipeline/Source/MovieRenderPipelineCore/Public/Graph/Nodes/MovieGraphRenderPassNode.h b/Engine/Plugins/MovieScene/MovieRenderPipeline/Source/MovieRenderPipelineCore/Public/Graph/Nodes/MovieGraphRenderPassNode.h
--- a/Engine/Plugins/MovieScene/MovieRenderPipeline/Source/MovieRenderPipelineCore/Public/Graph/Nodes/MovieGraphRenderPassNode.h
+++ b/Engine/Plugins/MovieScene/MovieRenderPipeline/Source/MovieRenderPipelineCore/Public/Graph/Nodes/MovieGraphRenderPassNode.h
@@ -12,6 +12,8 @@
 struct FMovieGraphRenderPassSetupData;
 struct FMovieGraphTimeStepData;
 
+#define UE_API MOVIERENDERPIPELINECORE_API
+
 /**
 * The UMovieGraphRenderPassNode node defines a render pass that MRQ may produce. This node can be implemented
 * in the graph multiple times, and the exact settings it should use can be created out of a mixture of nodes. Because
@@ -72,6 +74,14 @@ public:
 	 */
 	TArray<FSoftClassPath> GetOutputTypeRestrictions() const { return GetOutputTypeRestrictionsImpl(); }
 
+	UE_API virtual void ResolveTokenContainingProperties(TFunction<void(FString&)>& ResolveFunc, const FMovieGraphTokenResolveContext& InContext) override
+	{
+		Super::ResolveTokenContainingProperties(ResolveFunc, InContext);
+
+		ResolveFunc(RendererName);
+		ResolveFunc(RendererSubName);
+	}
+
 #if WITH_EDITOR
 	virtual FText GetMenuCategory() const override
 	{
@@ -116,3 +126,5 @@ protected:
 	virtual int32 GetNumSceneViewsRenderedImpl() const { return 0; }
 	virtual TArray<FSoftClassPath> GetOutputTypeRestrictionsImpl() const { return {}; } 
 };
+
+#undef UE_API
\ No newline at end of file
diff --git a/Engine/Plugins/MovieScene/MovieRenderPipeline/Source/MovieRenderPipelineRenderPasses/Private/Graph/Nodes/MovieGraphDeferredPassNode.cpp b/Engine/Plugins/MovieScene/MovieRenderPipeline/Source/MovieRenderPipelineRenderPasses/Private/Graph/Nodes/MovieGraphDeferredPassNode.cpp
--- a/Engine/Plugins/MovieScene/MovieRenderPipeline/Source/MovieRenderPipelineRenderPasses/Private/Graph/Nodes/MovieGraphDeferredPassNode.cpp
+++ b/Engine/Plugins/MovieScene/MovieRenderPipeline/Source/MovieRenderPipelineRenderPasses/Private/Graph/Nodes/MovieGraphDeferredPassNode.cpp
@@ -71,6 +71,16 @@ void UMovieGraphDeferredRenderPassNode::UpdateTelemetry(FMoviePipelineShotRender
 	InTelemetry->bUsesPageToSystemMemory |= bPageToSystemMemory;
 }
 
+void UMovieGraphDeferredRenderPassNode::ResolveTokenContainingProperties(TFunction<void(FString&)>& ResolveFunc, const FMovieGraphTokenResolveContext& InContext)
+{
+	Super::ResolveTokenContainingProperties(ResolveFunc, InContext);
+
+	for (FMoviePipelinePostProcessPass& PostProcessPass : AdditionalPostProcessMaterials)
+	{
+		ResolveFunc(PostProcessPass.Name);
+	}
+}
+
 #if WITH_EDITOR
 FText UMovieGraphDeferredRenderPassNode::GetNodeTitle(const bool bGetDescriptive) const
 {
diff --git a/Engine/Plugins/MovieScene/MovieRenderPipeline/Source/MovieRenderPipelineRenderPasses/Private/Graph/Nodes/MovieGraphPathTracerPassNode.cpp b/Engine/Plugins/MovieScene/MovieRenderPipeline/Source/MovieRenderPipelineRenderPasses/Private/Graph/Nodes/MovieGraphPathTracerPassNode.cpp
--- a/Engine/Plugins/MovieScene/MovieRenderPipeline/Source/MovieRenderPipelineRenderPasses/Private/Graph/Nodes/MovieGraphPathTracerPassNode.cpp
+++ b/Engine/Plugins/MovieScene/MovieRenderPipeline/Source/MovieRenderPipelineRenderPasses/Private/Graph/Nodes/MovieGraphPathTracerPassNode.cpp
@@ -282,3 +282,13 @@ void UMovieGraphPathTracerRenderPassNode::UpdateTelemetry(FMoviePipelineShotRend
 	InTelemetry->bUsesPPMs |= Algo::AnyOf(AdditionalPostProcessMaterials, [](const FMoviePipelinePostProcessPass& Pass) { return Pass.bEnabled; });
 	InTelemetry->SpatialSampleCount = FMath::Max(InTelemetry->SpatialSampleCount, SpatialSampleCount);
 }
+
+void UMovieGraphPathTracerRenderPassNode::ResolveTokenContainingProperties(TFunction<void(FString&)>& ResolveFunc, const FMovieGraphTokenResolveContext& InContext)
+{
+	Super::ResolveTokenContainingProperties(ResolveFunc, InContext);
+	
+	for (FMoviePipelinePostProcessPass& PostProcessPass : AdditionalPostProcessMaterials)
+	{
+		ResolveFunc(PostProcessPass.Name);
+	}
+}
diff --git a/Engine/Plugins/MovieScene/MovieRenderPipeline/Source/MovieRenderPipelineRenderPasses/Public/Graph/Nodes/MovieGraphDeferredPassNode.h b/Engine/Plugins/MovieScene/MovieRenderPipeline/Source/MovieRenderPipelineRenderPasses/Public/Graph/Nodes/MovieGraphDeferredPassNode.h
--- a/Engine/Plugins/MovieScene/MovieRenderPipeline/Source/MovieRenderPipelineRenderPasses/Public/Graph/Nodes/MovieGraphDeferredPassNode.h
+++ b/Engine/Plugins/MovieScene/MovieRenderPipeline/Source/MovieRenderPipelineRenderPasses/Public/Graph/Nodes/MovieGraphDeferredPassNode.h
@@ -20,6 +20,7 @@ public:
 
 	UE_API virtual void GetFormatResolveArgs(FMovieGraphResolveArgs& OutMergedFormatArgs, const FMovieGraphRenderDataIdentifier& InRenderDataIdentifier) const override;
 	UE_API virtual void UpdateTelemetry(FMoviePipelineShotRenderTelemetry* InTelemetry) const override;
+	UE_API virtual void ResolveTokenContainingProperties(TFunction<void(FString&)>& ResolveFunc, const FMovieGraphTokenResolveContext& InContext) override;
 
 #if WITH_EDITOR
 	UE_API virtual FText GetNodeTitle(const bool bGetDescriptive = false) const override;
diff --git a/Engine/Plugins/MovieScene/MovieRenderPipeline/Source/MovieRenderPipelineRenderPasses/Public/Graph/Nodes/MovieGraphPathTracerPassNode.h b/Engine/Plugins/MovieScene/MovieRenderPipeline/Source/MovieRenderPipelineRenderPasses/Public/Graph/Nodes/MovieGraphPathTracerPassNode.h
--- a/Engine/Plugins/MovieScene/MovieRenderPipeline/Source/MovieRenderPipelineRenderPasses/Public/Graph/Nodes/MovieGraphPathTracerPassNode.h
+++ b/Engine/Plugins/MovieScene/MovieRenderPipeline/Source/MovieRenderPipelineRenderPasses/Public/Graph/Nodes/MovieGraphPathTracerPassNode.h
@@ -71,6 +71,7 @@ public:
 	// ~UMovieGraphSettingNode Interface
 
 	UE_API virtual void UpdateTelemetry(FMoviePipelineShotRenderTelemetry* InTelemetry) const override;
+	UE_API virtual void ResolveTokenContainingProperties(TFunction<void(FString&)>& ResolveFunc, const FMovieGraphTokenResolveContext& InContext) override;
 
 protected:
 	// UMovieGraphRenderPassNode Interface
