MovieGraph: Fixed a number of issues with Metadata in MRG. Most per-camera metadata was missing from MRG exr files, but additionally when multiple layers were used they'd have conflicting keys and the metadata would have collisions which resulted in the loss of metadata. To resolve this, MRG has renamed per-layer metadata to be in the format "unreal/layerData/<exrLayerName>/<key name>", which avoids the name collisions.

Metadata can be looked up based on the layername in the exr (which may be "rgba", or something more complicated like Deferred_CamA, depending on what other layers are included in your exr file). The following Nuke TCL script can be used in a Text node to overlay data about the layer; By changing the output channel the Text node writes to, the metadata will look up the appropriate per-layer metadata for the matching pixel data.

[metadata exr/unreal/layerData/[value this.output]/branchName]
[metadata exr/unreal/layerData/[value this.output]/rendererName]
[metadata exr/unreal/layerData/[value this.output]/cameraName]
[metadata exr/unreal/layerData/[value this.output]/layerName]

Additionally made the binding-lookup for cameras slightly more robust to handle scenarios where Sequencer's Camera Cut track was bound to the component directly, instead of the root actor containing a camera component.

#virtualized

diff --git a/Engine/Plugins/MovieScene/MoviePipelineMaskRenderPass/Source/MoviePipelineMaskRenderPass/Private/Graph/MovieGraphObjectIdNode.cpp b/Engine/Plugins/MovieScene/MoviePipelineMaskRenderPass/Source/MoviePipelineMaskRenderPass/Private/Graph/MovieGraphObjectIdNode.cpp
--- a/Engine/Plugins/MovieScene/MoviePipelineMaskRenderPass/Source/MoviePipelineMaskRenderPass/Private/Graph/MovieGraphObjectIdNode.cpp
+++ b/Engine/Plugins/MovieScene/MoviePipelineMaskRenderPass/Source/MoviePipelineMaskRenderPass/Private/Graph/MovieGraphObjectIdNode.cpp
@@ -5,6 +5,7 @@
 #include "Editor/EditorPerProjectUserSettings.h"
 #include "MovieGraphObjectIdPass.h"
 #include "MoviePipelineTelemetry.h"
+#include "MoviePipelineUtils.h"
 
 #include UE_INLINE_GENERATED_CPP_BY_NAME(MovieGraphObjectIdNode)
 
@@ -32,6 +33,20 @@ FEngineShowFlags UMovieGraphObjectIdNode::GetShowFlags() const
 	return Flags;
 }
 
+void UMovieGraphObjectIdNode::GetFormatResolveArgs(FMovieGraphResolveArgs& OutMergedFormatArgs, const FMovieGraphRenderDataIdentifier& InRenderDataIdentifier) const
+{
+	FString MetadataPrefix = UE::MoviePipeline::GetMetadataPrefixPath(InRenderDataIdentifier);
+
+	OutMergedFormatArgs.FilenameArguments.Add(TEXT("ss_count"), FString::FromInt(SpatialSampleCount));
+	OutMergedFormatArgs.FileMetadata.Add(FString::Printf(TEXT("%s/spatialSampleCount"), *MetadataPrefix), FString::FromInt(SpatialSampleCount));
+
+	OutMergedFormatArgs.FilenameArguments.Add(TEXT("id_type"), UEnum::GetValueAsString(IdType));
+	OutMergedFormatArgs.FileMetadata.Add(FString::Printf(TEXT("%s/id_type"), *MetadataPrefix), UEnum::GetValueAsString(IdType));
+
+	OutMergedFormatArgs.FileMetadata.Add(FString::Printf(TEXT("%s/includeTranslucent"), *MetadataPrefix), FString::FromInt(bIncludeTranslucentObjects));
+}
+
+
 EViewModeIndex UMovieGraphObjectIdNode::GetViewModeIndex() const
 {
 	return VMI_Unlit;
diff --git a/Engine/Plugins/MovieScene/MoviePipelineMaskRenderPass/Source/MoviePipelineMaskRenderPass/Public/Graph/MovieGraphObjectIdNode.h b/Engine/Plugins/MovieScene/MoviePipelineMaskRenderPass/Source/MoviePipelineMaskRenderPass/Public/Graph/MovieGraphObjectIdNode.h
--- a/Engine/Plugins/MovieScene/MoviePipelineMaskRenderPass/Source/MoviePipelineMaskRenderPass/Public/Graph/MovieGraphObjectIdNode.h
+++ b/Engine/Plugins/MovieScene/MoviePipelineMaskRenderPass/Source/MoviePipelineMaskRenderPass/Public/Graph/MovieGraphObjectIdNode.h
@@ -28,6 +28,7 @@ public:
 	UE_API virtual bool GetForceLosslessCompression() const override;
 	virtual int32 GetNumSpatialSamples() const override { return SpatialSampleCount; }
 	virtual int32 GetNumSpatialSamplesDuringWarmUp() const  override { return 0; }
+	UE_API virtual void GetFormatResolveArgs(FMovieGraphResolveArgs& OutMergedFormatArgs, const FMovieGraphRenderDataIdentifier& InRenderDataIdentifier) const override;
 	// The ObjectID pass doesn't support any anti-aliasing so we force it to None so that their spatial samples turn into MRG-based camera jitter.
 	virtual bool GetOverrideAntiAliasing() const override { return true; }
 	virtual EAntiAliasingMethod GetAntiAliasingMethod() const override { return EAntiAliasingMethod::AAM_None; }
diff --git a/Engine/Plugins/MovieScene/MovieRenderPipeline/Source/MovieRenderPipelineCore/Private/Graph/MovieGraphBlueprintLibrary.cpp b/Engine/Plugins/MovieScene/MovieRenderPipeline/Source/MovieRenderPipelineCore/Private/Graph/MovieGraphBlueprintLibrary.cpp
--- a/Engine/Plugins/MovieScene/MovieRenderPipeline/Source/MovieRenderPipelineCore/Private/Graph/MovieGraphBlueprintLibrary.cpp
+++ b/Engine/Plugins/MovieScene/MovieRenderPipeline/Source/MovieRenderPipelineCore/Private/Graph/MovieGraphBlueprintLibrary.cpp
@@ -150,12 +150,12 @@ namespace
 			GeneratedResolveArgs.FilenameArguments.Add(TEXT("layer_name"), InParams.RenderDataIdentifier.LayerName);
 			GeneratedResolveArgs.FilenameArguments.Add(TEXT("branch_name"), InParams.RenderDataIdentifier.RootBranchName.ToString());
 
-			// TODO: Some of these are per render layer and need to be stored that way. EXRs will have the metadata for all the layers/cameras/etc. in one file.
-			GeneratedResolveArgs.FileMetadata.Add(TEXT("unreal/cameraName"), InParams.RenderDataIdentifier.CameraName);
-			GeneratedResolveArgs.FileMetadata.Add(TEXT("unreal/rendererName"), InParams.RenderDataIdentifier.RendererName);
-			GeneratedResolveArgs.FileMetadata.Add(TEXT("unreal/subResourceName"), InParams.RenderDataIdentifier.SubResourceName);
-			GeneratedResolveArgs.FileMetadata.Add(TEXT("unreal/layerName"), InParams.RenderDataIdentifier.LayerName);
-			GeneratedResolveArgs.FileMetadata.Add(TEXT("unreal/branchName"), InParams.RenderDataIdentifier.RootBranchName.ToString());
+			FString MetadataPath = UE::MoviePipeline::GetMetadataPrefixPath(InParams.RenderDataIdentifier);
+			GeneratedResolveArgs.FileMetadata.Add(FString::Printf(TEXT("%s/cameraName"), *MetadataPath), InParams.RenderDataIdentifier.CameraName);
+			GeneratedResolveArgs.FileMetadata.Add(FString::Printf(TEXT("%s/rendererName"), *MetadataPath), InParams.RenderDataIdentifier.RendererName);
+			GeneratedResolveArgs.FileMetadata.Add(FString::Printf(TEXT("%s/subResourceName"), *MetadataPath), InParams.RenderDataIdentifier.SubResourceName);
+			GeneratedResolveArgs.FileMetadata.Add(FString::Printf(TEXT("%s/layerName"), *MetadataPath), InParams.RenderDataIdentifier.LayerName);
+			GeneratedResolveArgs.FileMetadata.Add(FString::Printf(TEXT("%s/branchName"), *MetadataPath), InParams.RenderDataIdentifier.RootBranchName.ToString());
 		}
 
 
diff --git a/Engine/Plugins/MovieScene/MovieRenderPipeline/Source/MovieRenderPipelineCore/Private/Graph/MovieGraphDefaultRenderer.cpp b/Engine/Plugins/MovieScene/MovieRenderPipeline/Source/MovieRenderPipelineCore/Private/Graph/MovieGraphDefaultRenderer.cpp
--- a/Engine/Plugins/MovieScene/MovieRenderPipeline/Source/MovieRenderPipelineCore/Private/Graph/MovieGraphDefaultRenderer.cpp
+++ b/Engine/Plugins/MovieScene/MovieRenderPipeline/Source/MovieRenderPipelineCore/Private/Graph/MovieGraphDefaultRenderer.cpp
@@ -234,6 +234,7 @@ void UMovieGraphDefaultRenderer::AddReferencedObjects(UObject* InThis, FReferenc
 
 void UMovieGraphDefaultRenderer::Render(const FMovieGraphTimeStepData& InTimeStepData)
 {
+	// Flush built-in systems before we render anything. This maximizes the likelihood that the data is prepared for when
 	// the render thread uses it.
 	FlushAsyncEngineSystems(InTimeStepData.EvaluatedConfig);
 
diff --git a/Engine/Plugins/MovieScene/MovieRenderPipeline/Source/MovieRenderPipelineCore/Private/MoviePipelineBlueprintLibrary.cpp b/Engine/Plugins/MovieScene/MovieRenderPipeline/Source/MovieRenderPipelineCore/Private/MoviePipelineBlueprintLibrary.cpp
--- a/Engine/Plugins/MovieScene/MovieRenderPipeline/Source/MovieRenderPipelineCore/Private/MoviePipelineBlueprintLibrary.cpp
+++ b/Engine/Plugins/MovieScene/MovieRenderPipeline/Source/MovieRenderPipelineCore/Private/MoviePipelineBlueprintLibrary.cpp
@@ -563,10 +563,14 @@ void UMoviePipelineBlueprintLibrary::UpdateJobShotListFromSequence(ULevelSequenc
 						SidecarCamera.SequenceId = Entity.SequenceID; 
 						bool bAddCamera = true;
 						
+						TArray<FGuid> CheckedBindings;
+						
 						// We break out of the while loop below when the parent guid is invalid, so we track
 						// the last known valid GUID too, so we can compare the root binding against the camera.
 						FGuid LastValidGuid = Possessable.GetGuid();
 						FGuid ParentGuid = Possessable.GetParent();
+						CheckedBindings.Add(LastValidGuid);
+
 						while (ParentGuid.IsValid())
 						{
 							if (FMovieScenePossessable* ParentAsPossessable = LeafNode->MovieScene->FindPossessable(ParentGuid))
@@ -574,6 +578,7 @@ void UMoviePipelineBlueprintLibrary::UpdateJobShotListFromSequence(ULevelSequenc
 								LastValidGuid = ParentGuid;
 								ParentGuid = ParentAsPossessable->GetParent();
 								SidecarCamera.Name = ParentAsPossessable->GetName();
+								CheckedBindings.Add(LastValidGuid);
 
 								if (MovieSceneHelpers::IsBoundToAnySpawnable(LeafNode->MovieScene->GetTypedOuter<UMovieSceneSequence>(), ParentAsPossessable->GetGuid(), TransientPlaybackState))
 								{
@@ -610,7 +615,7 @@ void UMoviePipelineBlueprintLibrary::UpdateJobShotListFromSequence(ULevelSequenc
 							// This comparison doesn't support cross-sequence references, technically we should create a 
 							// FMovieSceneObjectBindingID etc, but because we only support finding sidecar cameras in the
 							// same sequence as the main camera cut track, this is fine for now.
-							if (LastValidGuid == MainBinding.GetGuid())
+							if (CheckedBindings.Contains(MainBinding.GetGuid()))
 							{
 								// We want to know which sidecar camera this is, not the possessable index.
 								MainCameraIndex = Entity.SidecarCameras.Num();
diff --git a/Engine/Plugins/MovieScene/MovieRenderPipeline/Source/MovieRenderPipelineCore/Private/MoviePipelineUtils.cpp b/Engine/Plugins/MovieScene/MovieRenderPipeline/Source/MovieRenderPipelineCore/Private/MoviePipelineUtils.cpp
--- a/Engine/Plugins/MovieScene/MovieRenderPipeline/Source/MovieRenderPipelineCore/Private/MoviePipelineUtils.cpp
+++ b/Engine/Plugins/MovieScene/MovieRenderPipeline/Source/MovieRenderPipelineCore/Private/MoviePipelineUtils.cpp
@@ -1145,23 +1145,31 @@ namespace UE
 			InFileMetadata.Add(TEXT("unreal/stats/memory/peakUsedPhysicalMB"), LexToString(MemoryStats.PeakUsedPhysical / MBDivider));
 			InFileMetadata.Add(TEXT("unreal/stats/memory/peakUsedVirtualMB"), LexToString(MemoryStats.PeakUsedVirtual / MBDivider));
 		}
-		
+
 		void GetMetadataFromCineCamera(UCineCameraComponent* InComponent, const FString& InCameraName, const FString& InRenderPassName, TMap<FString, FString>& InOutMetadata)
 		{
-			if (InComponent)
-			{
-				InOutMetadata.Add(FString::Printf(TEXT("unreal/%s/%s/sensorWidth"), *InCameraName, *InRenderPassName), FString::SanitizeFloat(InComponent->Filmback.SensorWidth));
-				InOutMetadata.Add(FString::Printf(TEXT("unreal/%s/%s/sensorHeight"), *InCameraName, *InRenderPassName), FString::SanitizeFloat(InComponent->Filmback.SensorHeight));
-				InOutMetadata.Add(FString::Printf(TEXT("unreal/%s/%s/sensorAspectRatio"), *InCameraName, *InRenderPassName), FString::SanitizeFloat(InComponent->Filmback.SensorAspectRatio));
-				InOutMetadata.Add(FString::Printf(TEXT("unreal/%s/%s/minFocalLength"), *InCameraName, *InRenderPassName), FString::SanitizeFloat(InComponent->LensSettings.MinFocalLength));
-				InOutMetadata.Add(FString::Printf(TEXT("unreal/%s/%s/maxFocalLength"), *InCameraName, *InRenderPassName), FString::SanitizeFloat(InComponent->LensSettings.MaxFocalLength));
-				InOutMetadata.Add(FString::Printf(TEXT("unreal/%s/%s/minFStop"), *InCameraName, *InRenderPassName), FString::SanitizeFloat(InComponent->LensSettings.MinFStop));
-				InOutMetadata.Add(FString::Printf(TEXT("unreal/%s/%s/maxFStop"), *InCameraName, *InRenderPassName), FString::SanitizeFloat(InComponent->LensSettings.MaxFStop));
-				InOutMetadata.Add(FString::Printf(TEXT("unreal/%s/%s/dofDiaphragmBladeCount"), *InCameraName, *InRenderPassName), FString::FromInt(InComponent->LensSettings.DiaphragmBladeCount));
-				InOutMetadata.Add(FString::Printf(TEXT("unreal/%s/%s/focalLength"), *InCameraName, *InRenderPassName), FString::SanitizeFloat(InComponent->CurrentFocalLength));
-			}
+			FString Prefix = FString::Printf(TEXT("unreal/%s/%s"), *InCameraName, *InRenderPassName);
+			GetMetadataFromCineCamera(InComponent, Prefix, InOutMetadata);
 		}
 		
+		void GetMetadataFromCineCamera(UCineCameraComponent* InComponent, const FString& InPrefix, TMap<FString, FString>& InOutMetadata)
+		{
+			if (!InComponent) 
+			{ 
+				return; 
+			}
+
+			InOutMetadata.Add(FString::Printf(TEXT("%s/sensorWidth"), *InPrefix), FString::SanitizeFloat(InComponent->Filmback.SensorWidth));
+			InOutMetadata.Add(FString::Printf(TEXT("%s/sensorHeight"), *InPrefix), FString::SanitizeFloat(InComponent->Filmback.SensorHeight));
+			InOutMetadata.Add(FString::Printf(TEXT("%s/sensorAspectRatio"), *InPrefix), FString::SanitizeFloat(InComponent->Filmback.SensorAspectRatio));
+			InOutMetadata.Add(FString::Printf(TEXT("%s/minFocalLength"), *InPrefix), FString::SanitizeFloat(InComponent->LensSettings.MinFocalLength));
+			InOutMetadata.Add(FString::Printf(TEXT("%s/maxFocalLength"), *InPrefix), FString::SanitizeFloat(InComponent->LensSettings.MaxFocalLength));
+			InOutMetadata.Add(FString::Printf(TEXT("%s/minFStop"), *InPrefix), FString::SanitizeFloat(InComponent->LensSettings.MinFStop));
+			InOutMetadata.Add(FString::Printf(TEXT("%s/maxFStop"), *InPrefix), FString::SanitizeFloat(InComponent->LensSettings.MaxFStop));
+			InOutMetadata.Add(FString::Printf(TEXT("%s/dofDiaphragmBladeCount"), *InPrefix), FString::FromInt(InComponent->LensSettings.DiaphragmBladeCount));
+			InOutMetadata.Add(FString::Printf(TEXT("%s/focalLength"), *InPrefix), FString::SanitizeFloat(InComponent->CurrentFocalLength));
+		}
+
 		void GetMetadataFromCameraLocRot(const FString& InCameraName, const FString& InRenderPassName, const FVector& InCurLoc, const FRotator& InCurRot, const FVector& InPrevLoc, const FRotator& InPrevRot, TMap<FString, FString>& InOutMetadata)
 		{
 			// InRenderPassName could be empty (for global camera stuff), and to support proper backwards compat metadata
@@ -1171,19 +1179,26 @@ namespace UE
 			{
 				CamName = InCameraName;
 			}
-			InOutMetadata.Add(FString::Printf(TEXT("unreal/%s/curPos/x"), *CamName), FString::SanitizeFloat(InCurLoc.X));
-			InOutMetadata.Add(FString::Printf(TEXT("unreal/%s/curPos/y"), *CamName), FString::SanitizeFloat(InCurLoc.Y));
-			InOutMetadata.Add(FString::Printf(TEXT("unreal/%s/curPos/z"), *CamName), FString::SanitizeFloat(InCurLoc.Z));
-			InOutMetadata.Add(FString::Printf(TEXT("unreal/%s/curRot/pitch"), *CamName),  FString::SanitizeFloat(InCurRot.Pitch));
-			InOutMetadata.Add(FString::Printf(TEXT("unreal/%s/curRot/yaw"), *CamName), FString::SanitizeFloat(InCurRot.Yaw));
-			InOutMetadata.Add(FString::Printf(TEXT("unreal/%s/curRot/roll"), *CamName), FString::SanitizeFloat(InCurRot.Roll));
-
-			InOutMetadata.Add(FString::Printf(TEXT("unreal/%s/prevPos/x"), *CamName), FString::SanitizeFloat(InPrevLoc.X));
-			InOutMetadata.Add(FString::Printf(TEXT("unreal/%s/prevPos/y"), *CamName), FString::SanitizeFloat(InPrevLoc.Y));
-			InOutMetadata.Add(FString::Printf(TEXT("unreal/%s/prevPos/z"), *CamName), FString::SanitizeFloat(InPrevLoc.Z));
-			InOutMetadata.Add(FString::Printf(TEXT("unreal/%s/prevRot/pitch"), *CamName), FString::SanitizeFloat(InPrevRot.Pitch));
-			InOutMetadata.Add(FString::Printf(TEXT("unreal/%s/prevRot/yaw"), *CamName), FString::SanitizeFloat(InPrevRot.Yaw));
-			InOutMetadata.Add(FString::Printf(TEXT("unreal/%s/prevRot/roll"), *CamName), FString::SanitizeFloat(InPrevRot.Roll));
+
+			FString Prefix = FString::Printf(TEXT("unreal/%s"), *CamName);
+			GetMetadataFromCameraLocRot(InCurLoc, InCurRot, InPrevLoc, InPrevRot, Prefix, InOutMetadata);
+		}
+		
+		void GetMetadataFromCameraLocRot(const FVector& InCurLoc, const FRotator& InCurRot, const FVector& InPrevLoc, const FRotator& InPrevRot, const FString& InPrefix, TMap<FString, FString>& InOutMetadata)
+		{
+			InOutMetadata.Add(FString::Printf(TEXT("%s/curPos/x"), *InPrefix), FString::SanitizeFloat(InCurLoc.X));
+			InOutMetadata.Add(FString::Printf(TEXT("%s/curPos/y"), *InPrefix), FString::SanitizeFloat(InCurLoc.Y));
+			InOutMetadata.Add(FString::Printf(TEXT("%s/curPos/z"), *InPrefix), FString::SanitizeFloat(InCurLoc.Z));
+			InOutMetadata.Add(FString::Printf(TEXT("%s/curRot/pitch"), *InPrefix),  FString::SanitizeFloat(InCurRot.Pitch));
+			InOutMetadata.Add(FString::Printf(TEXT("%s/curRot/yaw"), *InPrefix), FString::SanitizeFloat(InCurRot.Yaw));
+			InOutMetadata.Add(FString::Printf(TEXT("%s/curRot/roll"), *InPrefix), FString::SanitizeFloat(InCurRot.Roll));
+
+			InOutMetadata.Add(FString::Printf(TEXT("%s/prevPos/x"), *InPrefix), FString::SanitizeFloat(InPrevLoc.X));
+			InOutMetadata.Add(FString::Printf(TEXT("%s/prevPos/y"), *InPrefix), FString::SanitizeFloat(InPrevLoc.Y));
+			InOutMetadata.Add(FString::Printf(TEXT("%s/prevPos/z"), *InPrefix), FString::SanitizeFloat(InPrevLoc.Z));
+			InOutMetadata.Add(FString::Printf(TEXT("%s/prevRot/pitch"), *InPrefix), FString::SanitizeFloat(InPrevRot.Pitch));
+			InOutMetadata.Add(FString::Printf(TEXT("%s/prevRot/yaw"), *InPrefix), FString::SanitizeFloat(InPrevRot.Yaw));
+			InOutMetadata.Add(FString::Printf(TEXT("%s/prevRot/roll"), *InPrefix), FString::SanitizeFloat(InPrevRot.Roll));
 		}
 
 
@@ -1489,5 +1504,10 @@ namespace UE
 			}
 		}
 
+		FString GetMetadataPrefixPath(const FMovieGraphRenderDataIdentifier& Identifier)
+		{
+			return FString::Printf(TEXT("unreal/%s/%s/%s"), *Identifier.LayerName, *Identifier.RendererName, *Identifier.CameraName);
+		}
+
 	}
 }
diff --git a/Engine/Plugins/MovieScene/MovieRenderPipeline/Source/MovieRenderPipelineCore/Public/MoviePipelineUtils.h b/Engine/Plugins/MovieScene/MovieRenderPipeline/Source/MovieRenderPipelineCore/Public/MoviePipelineUtils.h
--- a/Engine/Plugins/MovieScene/MovieRenderPipeline/Source/MovieRenderPipelineCore/Public/MoviePipelineUtils.h
+++ b/Engine/Plugins/MovieScene/MovieRenderPipeline/Source/MovieRenderPipelineCore/Public/MoviePipelineUtils.h
@@ -222,8 +222,11 @@ namespace UE
 		MOVIERENDERPIPELINECORE_API void GetSharedFormatArguments(TMap<FString, FString>& InFilenameArguments, TMap<FString, FString>& InFileMetadata, const FDateTime& InDateTime, const int32 InVersionNumber, const UMoviePipelineExecutorJob* InJob, const FTimespan& InInitializationTimeOffset = FTimespan());
 		MOVIERENDERPIPELINECORE_API void GetHardwareUsageMetadata(TMap<FString, FString>& InFileMetadata, const FString& InOutputDir);
 		MOVIERENDERPIPELINECORE_API void GetDiagnosticMetadata(TMap<FString, FString>& InFileMetadata, const bool bIsGraph);
-		MOVIERENDERPIPELINECORE_API void GetMetadataFromCineCamera(class UCineCameraComponent* InComponent, const FString& InCameraName, const FString& InRenderPassName, TMap<FString, FString>& InOutMetadata);
+		MOVIERENDERPIPELINECORE_API void GetMetadataFromCineCamera(UCineCameraComponent* InComponent, const FString& InPrefix, TMap<FString, FString>& InOutMetadata);
+		MOVIERENDERPIPELINECORE_API void GetMetadataFromCineCamera(UCineCameraComponent* InComponent, const FString& InCameraName, const FString& InRenderPassName, TMap<FString, FString>& InOutMetadata);
 		MOVIERENDERPIPELINECORE_API void GetMetadataFromCameraLocRot(const FString& InCameraName, const FString& InRenderPassName, const FVector& InCurLoc, const FRotator& InCurRot, const FVector& InPrevLoc, const FRotator& InPrevRot, TMap<FString, FString>& InOutMetadata);
+		MOVIERENDERPIPELINECORE_API void GetMetadataFromCameraLocRot(const FVector& InCurLoc, const FRotator& InCurRot, const FVector& InPrevLoc, const FRotator& InPrevRot, const FString& InPrefix, TMap<FString, FString>& InOutMetadata);
+
 		MOVIERENDERPIPELINECORE_API FMoviePipelineRenderPassMetrics GetRenderPassMetrics(UMoviePipelinePrimaryConfig* InPrimaryConfig, UMoviePipelineExecutorShot* InPipelineExecutorShot, const FMoviePipelineRenderPassMetrics& InRenderPassMetrics, const FIntPoint& InEffectiveOutputResolution);
 		MOVIERENDERPIPELINECORE_API bool CanWriteToFile(const TCHAR* InFilename, bool bOverwriteExisting);
 		MOVIERENDERPIPELINECORE_API FString GetPaddingFormatString(int32 InZeroPadCount, const int32 InFrameNumber);
@@ -249,6 +252,10 @@ namespace UE
 			/** Auto-exposure cubemap passes are laid out in a 3x2 arrangement, and rounded down to a multiple of 8 pixels. */
 			MOVIERENDERPIPELINECORE_API int32 ComputeAutoExposureCubeCaptureSize(FIntPoint Resolution);
 		}
+
+		// When using EXR metadata what is the prefix that metadata tokens should go under? This returns "unreal/<layerName>/<rendererName>/<cameraName>" but be aware that the metadata keys will be potentially
+		// renamed when actually writing the EXR to disk so that the per-layer data goes into a key that matches the EXR layer names (which makes the lookup in external software much easier).
+		MOVIERENDERPIPELINECORE_API FString GetMetadataPrefixPath(const FMovieGraphRenderDataIdentifier& Identifier);
 	}
 }
 
diff --git a/Engine/Plugins/MovieScene/MovieRenderPipeline/Source/MovieRenderPipelineRenderPasses/Private/Graph/Nodes/MovieGraphDeferredPanoramicPassNode.cpp b/Engine/Plugins/MovieScene/MovieRenderPipeline/Source/MovieRenderPipelineRenderPasses/Private/Graph/Nodes/MovieGraphDeferredPanoramicPassNode.cpp
--- a/Engine/Plugins/MovieScene/MovieRenderPipeline/Source/MovieRenderPipelineRenderPasses/Private/Graph/Nodes/MovieGraphDeferredPanoramicPassNode.cpp
+++ b/Engine/Plugins/MovieScene/MovieRenderPipeline/Source/MovieRenderPipelineRenderPasses/Private/Graph/Nodes/MovieGraphDeferredPanoramicPassNode.cpp
@@ -3,6 +3,7 @@
 #include "Graph/Nodes/MovieGraphDeferredPanoramicPassNode.h"
 #include "Graph/Renderers/MovieGraphDeferredPanoramicPass.h"
 #include "MoviePipelineTelemetry.h"
+#include "MoviePipelineUtils.h"
 
 #include UE_INLINE_GENERATED_CPP_BY_NAME(MovieGraphDeferredPanoramicPassNode)
 
@@ -23,6 +24,32 @@ UMovieGraphDeferredPanoramicNode::UMovieGraphDeferredPanoramicNode()
 	RendererName = TEXT("DeferredPanoramic");
 }
 
+void UMovieGraphDeferredPanoramicNode::GetFormatResolveArgs(FMovieGraphResolveArgs& OutMergedFormatArgs, const FMovieGraphRenderDataIdentifier& InRenderDataIdentifier) const
+{
+	FString MetadataPrefix = UE::MoviePipeline::GetMetadataPrefixPath(InRenderDataIdentifier);
+
+	OutMergedFormatArgs.FilenameArguments.Add(TEXT("pano_count_h"), FString::FromInt(NumHorizontalSteps));
+	OutMergedFormatArgs.FileMetadata.Add(FString::Printf(TEXT("%s/paneCountH"), *MetadataPrefix), FString::FromInt(NumHorizontalSteps));
+
+	OutMergedFormatArgs.FilenameArguments.Add(TEXT("pano_count_v"), FString::FromInt(NumVerticalSteps));
+	OutMergedFormatArgs.FileMetadata.Add(FString::Printf(TEXT("%s/paneCountV"), *MetadataPrefix), FString::FromInt(NumVerticalSteps));
+
+	OutMergedFormatArgs.FilenameArguments.Add(TEXT("pano_follow_orientation"), FString::FromInt(bFollowCameraOrientation));
+	OutMergedFormatArgs.FileMetadata.Add(FString::Printf(TEXT("%s/followOrientation"), *MetadataPrefix), FString::FromInt(bFollowCameraOrientation));
+
+	OutMergedFormatArgs.FilenameArguments.Add(TEXT("aaMethod"), UEnum::GetValueAsString(AntiAliasingMethod));
+	OutMergedFormatArgs.FileMetadata.Add(FString::Printf(TEXT("%s/aaMethod"), *MetadataPrefix), UEnum::GetValueAsString(AntiAliasingMethod));
+
+	OutMergedFormatArgs.FilenameArguments.Add(TEXT("pano_filter"), UEnum::GetValueAsString(Filter));
+	OutMergedFormatArgs.FileMetadata.Add(FString::Printf(TEXT("%s/filter"), *MetadataPrefix), UEnum::GetValueAsString(Filter));
+
+	OutMergedFormatArgs.FilenameArguments.Add(TEXT("disable_tonecurve"), FString::FromInt(bDisableToneCurve));
+	OutMergedFormatArgs.FileMetadata.Add(FString::Printf(TEXT("%s/disableTonecurve"), *MetadataPrefix), FString::FromInt(bDisableToneCurve));
+
+	OutMergedFormatArgs.FilenameArguments.Add(TEXT("history_per_pane"), FString::FromInt(bAllocateHistoryPerPane));
+	OutMergedFormatArgs.FileMetadata.Add(FString::Printf(TEXT("%s/historyPerPane"), *MetadataPrefix), FString::FromInt(bAllocateHistoryPerPane));
+}
+
 FEngineShowFlags UMovieGraphDeferredPanoramicNode::GetShowFlags() const
 {
 	FEngineShowFlags Flags = Super::GetShowFlags();
diff --git a/Engine/Plugins/MovieScene/MovieRenderPipeline/Source/MovieRenderPipelineRenderPasses/Private/Graph/Nodes/MovieGraphDeferredPassNode.cpp b/Engine/Plugins/MovieScene/MovieRenderPipeline/Source/MovieRenderPipelineRenderPasses/Private/Graph/Nodes/MovieGraphDeferredPassNode.cpp
--- a/Engine/Plugins/MovieScene/MovieRenderPipeline/Source/MovieRenderPipelineRenderPasses/Private/Graph/Nodes/MovieGraphDeferredPassNode.cpp
+++ b/Engine/Plugins/MovieScene/MovieRenderPipeline/Source/MovieRenderPipelineRenderPasses/Private/Graph/Nodes/MovieGraphDeferredPassNode.cpp
@@ -4,6 +4,7 @@
 
 #include "Graph/Renderers/MovieGraphDeferredPass.h"
 #include "MoviePipelineTelemetry.h"
+#include "MoviePipelineUtils.h"
 
 #include UE_INLINE_GENERATED_CPP_BY_NAME(MovieGraphDeferredPassNode)
 
@@ -46,19 +47,23 @@ UMovieGraphDeferredRenderPassNode::UMovieGraphDeferredRenderPassNode()
 
 void UMovieGraphDeferredRenderPassNode::GetFormatResolveArgs(FMovieGraphResolveArgs& OutMergedFormatArgs, const FMovieGraphRenderDataIdentifier& InRenderDataIdentifier) const
 {
-	// ToDo: This should be split to write things per renderer as they can be different for each layer/camera/render pass.
+	FString MetadataPrefix = UE::MoviePipeline::GetMetadataPrefixPath(InRenderDataIdentifier);
+
 	// We intentionally skip some settings for not being very meaningful to output, ie: bAllowOCIO, bPageToSystemMemory, bWriteAllSamples
 	OutMergedFormatArgs.FilenameArguments.Add(TEXT("ss_count"), FString::FromInt(SpatialSampleCount));
-	OutMergedFormatArgs.FileMetadata.Add(TEXT("unreal/sampling/spatialSampleCount"), FString::FromInt(SpatialSampleCount));
+	OutMergedFormatArgs.FileMetadata.Add(FString::Printf(TEXT("%s/spatialSampleCount"), *MetadataPrefix), FString::FromInt(SpatialSampleCount));
 
 	OutMergedFormatArgs.FilenameArguments.Add(TEXT("disable_tonecurve"), FString::FromInt(bDisableToneCurve));
-	OutMergedFormatArgs.FileMetadata.Add(TEXT("unreal/disableTonecurve"), FString::FromInt(bDisableToneCurve));
+	OutMergedFormatArgs.FileMetadata.Add(FString::Printf(TEXT("%s/disableTonecurve"), *MetadataPrefix), FString::FromInt(bDisableToneCurve));
 
 	OutMergedFormatArgs.FilenameArguments.Add(TEXT("overlap_percentage"), FString::SanitizeFloat(OverlapPercentage));
-	OutMergedFormatArgs.FileMetadata.Add(TEXT("unreal/highres/overlapPercentage"), FString::SanitizeFloat(OverlapPercentage));
+	OutMergedFormatArgs.FileMetadata.Add(FString::Printf(TEXT("%s/highres/overlapPercentage"), *MetadataPrefix), FString::SanitizeFloat(OverlapPercentage));
 
 	OutMergedFormatArgs.FilenameArguments.Add(TEXT("history_per_tile"), FString::FromInt(bAllocateHistoryPerTile));
-	OutMergedFormatArgs.FileMetadata.Add(TEXT("unreal/highres/historyPerTile"), FString::FromInt(bAllocateHistoryPerTile));
+	OutMergedFormatArgs.FileMetadata.Add(FString::Printf(TEXT("%s/highres/historyPerTile"), *MetadataPrefix), FString::FromInt(bAllocateHistoryPerTile));
+	
+	OutMergedFormatArgs.FilenameArguments.Add(TEXT("aaMethod"), UEnum::GetValueAsString(AntiAliasingMethod));
+	OutMergedFormatArgs.FileMetadata.Add(FString::Printf(TEXT("%s/aaMethod"), *MetadataPrefix), UEnum::GetValueAsString(AntiAliasingMethod));
 }
 
 void UMovieGraphDeferredRenderPassNode::UpdateTelemetry(FMoviePipelineShotRenderTelemetry* InTelemetry) const
diff --git a/Engine/Plugins/MovieScene/MovieRenderPipeline/Source/MovieRenderPipelineRenderPasses/Private/Graph/Nodes/MovieGraphPathTracerPassNode.cpp b/Engine/Plugins/MovieScene/MovieRenderPipeline/Source/MovieRenderPipelineRenderPasses/Private/Graph/Nodes/MovieGraphPathTracerPassNode.cpp
--- a/Engine/Plugins/MovieScene/MovieRenderPipeline/Source/MovieRenderPipelineRenderPasses/Private/Graph/Nodes/MovieGraphPathTracerPassNode.cpp
+++ b/Engine/Plugins/MovieScene/MovieRenderPipeline/Source/MovieRenderPipelineRenderPasses/Private/Graph/Nodes/MovieGraphPathTracerPassNode.cpp
@@ -8,6 +8,7 @@
 #include "RenderUtils.h"
 #include "ShowFlags.h"
 #include "PathTracingDenoiser.h"
+#include "MoviePipelineUtils.h"
 
 #include UE_INLINE_GENERATED_CPP_BY_NAME(MovieGraphPathTracerPassNode)
 
@@ -52,11 +53,16 @@ UMovieGraphPathTracerRenderPassNode::UMovieGraphPathTracerRenderPassNode()
 
 void UMovieGraphPathTracerRenderPassNode::GetFormatResolveArgs(FMovieGraphResolveArgs& OutMergedFormatArgs, const FMovieGraphRenderDataIdentifier& InRenderDataIdentifier) const
 {
-	OutMergedFormatArgs.FilenameArguments.Add(TEXT("path_tracer_sampling_count"), FString::FromInt(SpatialSampleCount));
+	FString MetadataPrefix = UE::MoviePipeline::GetMetadataPrefixPath(InRenderDataIdentifier);
+	
+	OutMergedFormatArgs.FilenameArguments.Add(TEXT("ss_count"), FString::FromInt(SpatialSampleCount));
+	OutMergedFormatArgs.FileMetadata.Add(FString::Printf(TEXT("%s/spatialSampleCount"), *MetadataPrefix), FString::FromInt(SpatialSampleCount));
+	
 	OutMergedFormatArgs.FilenameArguments.Add(TEXT("path_tracer_seed_offset"), FString::FromInt(SeedOffset));
-
-	OutMergedFormatArgs.FileMetadata.Add(TEXT("unreal/pathTracer/samplingCount"), FString::FromInt(SpatialSampleCount));
-	OutMergedFormatArgs.FileMetadata.Add(TEXT("unreal/pathTracer/seedOffset"), FString::FromInt(SeedOffset));
+	OutMergedFormatArgs.FileMetadata.Add(FString::Printf(TEXT("%s/seedOffset"), *MetadataPrefix), FString::FromInt(SeedOffset));
+	
+	OutMergedFormatArgs.FilenameArguments.Add(TEXT("disable_tonecurve"), FString::FromInt(bDisableToneCurve));
+	OutMergedFormatArgs.FileMetadata.Add(FString::Printf(TEXT("%s/disableTonecurve"), *MetadataPrefix), FString::FromInt(bDisableToneCurve));
 }
 
 #if WITH_EDITOR
diff --git a/Engine/Plugins/MovieScene/MovieRenderPipeline/Source/MovieRenderPipelineRenderPasses/Private/Graph/Renderers/MovieGraphDeferredPass.cpp b/Engine/Plugins/MovieScene/MovieRenderPipeline/Source/MovieRenderPipelineRenderPasses/Private/Graph/Renderers/MovieGraphDeferredPass.cpp
--- a/Engine/Plugins/MovieScene/MovieRenderPipeline/Source/MovieRenderPipelineRenderPasses/Private/Graph/Renderers/MovieGraphDeferredPass.cpp
+++ b/Engine/Plugins/MovieScene/MovieRenderPipeline/Source/MovieRenderPipelineRenderPasses/Private/Graph/Renderers/MovieGraphDeferredPass.cpp
@@ -33,10 +33,11 @@ namespace UE::MovieGraph::Rendering
 {
 	namespace MetadataHelper
 	{
-		void UpdateSpatialSampleMetadata(const FSceneViewProjectionData& InProjectionData, TMap<FString, FString>& InOutMetadataMap)
+		void UpdateSpatialSampleMetadata(const FSceneViewProjectionData& InProjectionData, TMap<FString, FString>& InOutMetadataMap, const FMovieGraphRenderDataIdentifier& InIdentifier)
 		{
-			InOutMetadataMap.Add(FString::Printf(TEXT("unreal/sampling/matrix/worldToCamera")), InProjectionData.ViewRotationMatrix.ToString());
-			InOutMetadataMap.Add(FString::Printf(TEXT("unreal/sampling/matrix/worldToUEClip")), InProjectionData.ProjectionMatrix.ToString());
+			FString Prefix = UE::MoviePipeline::GetMetadataPrefixPath(InIdentifier);
+			InOutMetadataMap.Add(FString::Printf(TEXT("%s/matrix/worldToCamera"), *Prefix), InProjectionData.ViewRotationMatrix.ToString());
+			InOutMetadataMap.Add(FString::Printf(TEXT("%s/matrix/worldToUEClip"), *Prefix), InProjectionData.ProjectionMatrix.ToString());
 		}
 	}
 
@@ -481,7 +482,16 @@ void FMovieGraphDeferredPass::Render(const FMovieGraphTraversalContext& InFrameT
 
 				ApplyMovieGraphOverridesToSampleState(SampleState);
 
-				MetadataHelper::UpdateSpatialSampleMetadata(SceneViewInitOptions, SampleState.AdditionalFileMetadata);
+				MetadataHelper::UpdateSpatialSampleMetadata(SceneViewInitOptions, SampleState.AdditionalFileMetadata, RenderDataIdentifier);
+				FString MetadataPrefix = UE::MoviePipeline::GetMetadataPrefixPath(RenderDataIdentifier);
+					UE::MoviePipeline::GetMetadataFromCameraLocRot(
+					CameraInfo.ViewInfo.Location, CameraInfo.ViewInfo.Rotation, 
+					CameraInfo.ViewInfo.PreviousViewTransform.Get(FTransform::Identity).GetLocation(), CameraInfo.ViewInfo.PreviousViewTransform.Get(FTransform::Identity).GetRotation().Rotator(), 
+					MetadataPrefix, SampleState.AdditionalFileMetadata
+				);
+				
+				// ToDo: If an actor had multiple cine camera components this could be wrong, but Sequencer also can't disambiguate between multiple camera components on a single actor.
+				UE::MoviePipeline::GetMetadataFromCineCamera(CameraInfo.ViewActor->GetComponentByClass<UCineCameraComponent>(), MetadataPrefix, SampleState.AdditionalFileMetadata);
 
 				// If this was just to contribute to the history buffer, no need to go any further. Never discard if we're writing individual samples, though.
 				bool bRenderPass = (!InTimeData.bDiscardOutput && !ShouldDiscardOutput(ViewFamily, CameraInfo)) || SampleState.bWriteSampleToDisk;
diff --git a/Engine/Plugins/MovieScene/MovieRenderPipeline/Source/MovieRenderPipelineRenderPasses/Private/MovieGraphImageSequenceOutputNode.cpp b/Engine/Plugins/MovieScene/MovieRenderPipeline/Source/MovieRenderPipelineRenderPasses/Private/MovieGraphImageSequenceOutputNode.cpp
--- a/Engine/Plugins/MovieScene/MovieRenderPipeline/Source/MovieRenderPipelineRenderPasses/Private/MovieGraphImageSequenceOutputNode.cpp
+++ b/Engine/Plugins/MovieScene/MovieRenderPipeline/Source/MovieRenderPipelineRenderPasses/Private/MovieGraphImageSequenceOutputNode.cpp
@@ -63,6 +63,44 @@
 		OutInjectedNodes.Add(BurnInNode);																		\
 	}																											\
 
+namespace UE::MovieGraph
+{
+	void UpdateMetadataLayerNamesToMatchExrLayerNames(const FMovieGraphRenderDataIdentifier& InIdentifier, const FString& LayerName, TMap<FString, FStringFormatArg>& InOutMetadata)
+	{
+		// We rename the metadata in the file when we write it to disk, because there's no other way to know what metadata goes with the actual layer name in the exr,
+		// and we don't want to require complex string manipulation in pipeline tools, so we rename `unreal/<layerName>/<rendererName>/<cameraName>/foo` to just
+		// `unreal/layerData/<exrLayerName>/foo`, where exrLayerName comes from this function that decided if the current layer is either "rgba" or something more complex.
+		TMap<FString, FStringFormatArg> NewMetadataNames;
+		TArray<FString> OldMetadataKeys;
+
+		const FString MetadataPrefixForLayer = UE::MoviePipeline::GetMetadataPrefixPath(InIdentifier);
+
+		for (const TPair<FString, FStringFormatArg>& KVP : InOutMetadata)
+		{
+			if (KVP.Key.StartsWith(MetadataPrefixForLayer))
+			{
+				const FString NewMetadataPrefix = FString::Printf(TEXT("unreal/layerData/%s"), LayerName.IsEmpty() ? TEXT("rgba") : *LayerName);
+				const FString OldMetadataPostfix = KVP.Key.RightChop(MetadataPrefixForLayer.Len());
+				const FString NewKey = NewMetadataPrefix + OldMetadataPostfix;
+
+				NewMetadataNames.Add(NewKey, KVP.Value);
+
+				// We want to remove the old keys to prevent confusion
+				OldMetadataKeys.Add(KVP.Key);
+			}
+		}
+
+		// Remove the old keys from the map
+		for (const FString& OldKey : OldMetadataKeys)
+		{
+			InOutMetadata.Remove(OldKey);
+		}
+
+		// Add the renamed ones
+		InOutMetadata.Append(NewMetadataNames);
+	}
+}
+
 UMovieGraphImageSequenceOutputNode::UMovieGraphImageSequenceOutputNode()
 {
 	ImageWriteQueue = &FModuleManager::Get().LoadModuleChecked<IImageWriteQueueModule>("ImageWriteQueue").GetWriteQueue();
@@ -561,6 +599,10 @@ void UMovieGraphImageSequenceOutputNode_EXR::OnReceiveImageDataImpl(UMovieGraphP
 			}
 		}
 
+
+		// Rename the per-layer metadata to match the EXR layer names to make fetching metadata simple string lookups.
+		UE::MovieGraph::UpdateMetadataLayerNamesToMatchExrLayerNames(RenderData.Key, TEXT(""), ImageWriteTask->FileMetadata);
+
 		UE::MovieGraph::FMovieGraphOutputFutureData OutputFutureData;
 		OutputFutureData.Shot = InPipeline->GetActiveShotList()[Payload->TraversalContext.ShotIndex];
 		OutputFutureData.FilePath = FileName;
@@ -760,7 +802,10 @@ void UMovieGraphImageSequenceOutputNode_MultiLayerEXR::OnReceiveImageDataImpl(UM
 			// Ensure the write task uses the maximum resolution of all the layers being written
 			MultiLayerImageTask->Width = RenderConfig.MaximumResolution.X;
 			MultiLayerImageTask->Height = RenderConfig.MaximumResolution.Y;
-			
+
+			// Rename the per-layer metadata to match the EXR layer names to make fetching metadata simple string lookups.
+			UE::MovieGraph::UpdateMetadataLayerNamesToMatchExrLayerNames(RenderID, LayerName, MultiLayerImageTask->FileMetadata);
+
 			LayerIndex++;
 		}
 
diff --git a/Engine/Plugins/MovieScene/MovieRenderPipeline/Source/MovieRenderPipelineRenderPasses/Public/Graph/Nodes/MovieGraphDeferredPanoramicPassNode.h b/Engine/Plugins/MovieScene/MovieRenderPipeline/Source/MovieRenderPipelineRenderPasses/Public/Graph/Nodes/MovieGraphDeferredPanoramicPassNode.h
--- a/Engine/Plugins/MovieScene/MovieRenderPipeline/Source/MovieRenderPipelineRenderPasses/Public/Graph/Nodes/MovieGraphDeferredPanoramicPassNode.h
+++ b/Engine/Plugins/MovieScene/MovieRenderPipeline/Source/MovieRenderPipelineRenderPasses/Public/Graph/Nodes/MovieGraphDeferredPanoramicPassNode.h
@@ -27,6 +27,7 @@ public:
 	virtual bool GetOverrideAntiAliasing() const override { return bOverride_AntiAliasingMethod; }
 	UE_API virtual EAntiAliasingMethod GetAntiAliasingMethod() const override;
 	virtual bool GetEnableHistoryPerTile() const { return bAllocateHistoryPerPane; }
+	UE_API virtual void GetFormatResolveArgs(FMovieGraphResolveArgs& OutMergedFormatArgs, const FMovieGraphRenderDataIdentifier& InRenderDataIdentifier) const;
 	// ~UMovieGraphImagePassBaseNode Interface
 
 	// UMovieGraphNode Interface
