MoviePipeline: Added additional logic to MRG's filename disambiguation process that prevents {renderer_name} and {renderer_sub_name} needing to be added in many scenarios. Also added a new "PPM File Name Format" property to Deferred and Path Traced renderer nodes, providing additional naming flexibility.

For example, if only a deferred render and one Post Process Material pass are being generated, both {renderer_name} and {renderer_sub_name} are not required at all.

diff --git a/Engine/Plugins/MovieScene/MoviePipelineMaskRenderPass/Source/MoviePipelineMaskRenderPass/Private/Graph/MovieGraphObjectIdPass.cpp b/Engine/Plugins/MovieScene/MoviePipelineMaskRenderPass/Source/MoviePipelineMaskRenderPass/Private/Graph/MovieGraphObjectIdPass.cpp
--- a/Engine/Plugins/MovieScene/MoviePipelineMaskRenderPass/Source/MoviePipelineMaskRenderPass/Private/Graph/MovieGraphObjectIdPass.cpp
+++ b/Engine/Plugins/MovieScene/MoviePipelineMaskRenderPass/Source/MoviePipelineMaskRenderPass/Private/Graph/MovieGraphObjectIdPass.cpp
@@ -169,6 +169,9 @@ namespace UE::MovieGraph
 				// Don't allow OCIO, otherwise Cryptomatte colors will not be correct.
 				NewPayload->bAllowOCIO = false;
 
+				// Object ID passes are never beauty passes; each one is effectively a subresource without any parent pass.
+				NewPayload->bIsBeautyPass = false;
+
 				TUniquePtr<TImagePixelData<FLinearColor>> FinalPixelData = MakeUnique<TImagePixelData<FLinearColor>>(FIntPoint(FullSizeX, FullSizeY), MoveTemp(OutputLayers[Index]), NewPayload);
 
 				// Send each layer to the Output Builder
diff --git a/Engine/Plugins/MovieScene/MovieRenderPipeline/Source/MovieRenderPipelineCore/Private/Graph/MovieGraphPipeline.cpp b/Engine/Plugins/MovieScene/MovieRenderPipeline/Source/MovieRenderPipelineCore/Private/Graph/MovieGraphPipeline.cpp
--- a/Engine/Plugins/MovieScene/MovieRenderPipeline/Source/MovieRenderPipelineCore/Private/Graph/MovieGraphPipeline.cpp
+++ b/Engine/Plugins/MovieScene/MovieRenderPipeline/Source/MovieRenderPipelineCore/Private/Graph/MovieGraphPipeline.cpp
@@ -218,7 +218,7 @@ void UMovieGraphPipeline::DuplicateJobAndConfiguration()
 	// We only look in the primary configuration for the job for script nodes (and not shot specific overrides). If we looked
 	// in the shot-specific overrides, we would end up creating instances of the scripts for those shots. This can become really
 	// confusing if a shot specifies an Execute Script node with the same script as the Primary Graph, now you'll have two separate
-	// instances of your script, one which recieves 4 callbacks (pre/post job, pre/post shot) and one that only recieves pre/post shot.
+	// instances of your script, one which receives 4 callbacks (pre/post job, pre/post shot) and one that only receives pre/post shot.
 	FMovieGraphTraversalContext Context;
 	Context.Job = CurrentJobDuplicate;
 	Context.RootGraph = DuplicatePrimaryConfig;
@@ -1549,10 +1549,11 @@ void UMovieGraphPipeline::ProcessOutstandingFinishedFrames()
 		// output containers for the data in this output frame, and a mask which lets that container know if it should skip a layer.
 		TMap<UMovieGraphFileOutputNode*, TSet<FMovieGraphRenderDataIdentifier>> MaskData;
 
-		for (TPair<FMovieGraphRenderDataIdentifier, TUniquePtr<FImagePixelData>>& RenderData : OutputFrame.ImageOutputData)
+		// For each layer in the output merger frame...
+		for (TPair<FMovieGraphRenderDataIdentifier, TUniquePtr<FImagePixelData>>& LayerData : OutputFrame.ImageOutputData)
 		{
-			UE::MovieGraph::FMovieGraphSampleState* Payload = RenderData.Value->GetPayload<UE::MovieGraph::FMovieGraphSampleState>();
-			if (!ensure(Payload))
+			UE::MovieGraph::FMovieGraphSampleState* LayerPayload = LayerData.Value->GetPayload<UE::MovieGraph::FMovieGraphSampleState>();
+			if (!ensure(LayerPayload))
 			{
 				continue;
 			}
@@ -1562,36 +1563,36 @@ void UMovieGraphPipeline::ProcessOutstandingFinishedFrames()
 				continue;
 			}
 
+			// For each File Output Node CDO in the layer, associate the CDO with the layer...
+			//
 			// Get a list of all output nodes for this particular render layer. We specifically skip CDOs here because we're just trying
 			// to find out which output nodes the user _wanted_ to place items onto, we later collect only the CDOs so that we have
 			// a central point for actually handling the file writing.
 			constexpr bool bIncludeCDOs = false;
-			const FName BranchName = RenderData.Key.RootBranchName;
-			TArray<UMovieGraphFileOutputNode*> OutputNodeInstances = OutputFrame.EvaluatedConfig->GetSettingsForBranch<UMovieGraphFileOutputNode>(BranchName, bIncludeCDOs);
-			for (const UMovieGraphFileOutputNode* Instance : OutputNodeInstances)
+			const FMovieGraphRenderDataIdentifier& LayerIdentifier = LayerData.Key;
+			const FName& LayerBranchName = LayerIdentifier.RootBranchName;
+			TArray<UMovieGraphFileOutputNode*> OutputNodesForLayer = OutputFrame.EvaluatedConfig->GetSettingsForBranch<UMovieGraphFileOutputNode>(LayerBranchName, bIncludeCDOs);
+			for (const UMovieGraphFileOutputNode* LayerOutputNode : OutputNodesForLayer)
 			{
-				UMovieGraphFileOutputNode* CDO = Instance->GetClass()->GetDefaultObject<UMovieGraphFileOutputNode>();
-				MaskData.FindOrAdd(CDO).Add(RenderData.Key);
+				UMovieGraphFileOutputNode* CDO = LayerOutputNode->GetClass()->GetDefaultObject<UMovieGraphFileOutputNode>();
+				MaskData.FindOrAdd(CDO).Add(LayerIdentifier);
 			}
 		}
 
 		// Now that we've looped through the above, we have the total list of which output formats are being used by the graph for
 		// all of the render layers given. We also have a list of which identifiers should go into each one. So we can loop through
 		// the CDO instances and pass the data to them.
-		//UE_LOG(LogMovieRenderPipeline, Warning, TEXT("File Outputs:"));
 		for (const TPair<UMovieGraphFileOutputNode*, TSet<FMovieGraphRenderDataIdentifier>>& Pair : MaskData)
 		{
-			//UE_LOG(LogMovieRenderPipeline, Warning, TEXT("\tNode: %s"), *Pair.Key->GetClass()->GetName());
-			for (const FMovieGraphRenderDataIdentifier& ID : Pair.Value)
-			{
-				//UE_LOG(LogMovieRenderPipeline, Warning, TEXT("\t\tBranch: %s:"), *ID.RootBranchName.ToString());
-			}
+			UMovieGraphFileOutputNode* OutputCDO = Pair.Key;
+			const TSet<FMovieGraphRenderDataIdentifier>& LayerIdentifiers = Pair.Value;
 
-			Pair.Key->OnReceiveImageData(this, &OutputFrame, Pair.Value);
+			// Tell the CDO that it's getting data from the output merger frame, and specify which layers it should output data for
+			Pair.Key->OnReceiveImageData(this, &OutputFrame, LayerIdentifiers);
 
 			// Ensure we keep track of which nodes we actually sent data to during this render
 			// so that we can call BeginFinalize/IsFinishedWritingToDisk on them when shutting down.
-			OutputNodesDataSentTo.Add(Pair.Key);
+			OutputNodesDataSentTo.Add(OutputCDO);
 		}
 	}
 }
diff --git a/Engine/Plugins/MovieScene/MovieRenderPipeline/Source/MovieRenderPipelineCore/Private/Graph/Nodes/MovieGraphFileOutputNode.cpp b/Engine/Plugins/MovieScene/MovieRenderPipeline/Source/MovieRenderPipelineCore/Private/Graph/Nodes/MovieGraphFileOutputNode.cpp
--- a/Engine/Plugins/MovieScene/MovieRenderPipeline/Source/MovieRenderPipelineCore/Private/Graph/Nodes/MovieGraphFileOutputNode.cpp
+++ b/Engine/Plugins/MovieScene/MovieRenderPipeline/Source/MovieRenderPipelineCore/Private/Graph/Nodes/MovieGraphFileOutputNode.cpp
@@ -2,7 +2,9 @@
 
 #include "Graph/Nodes/MovieGraphFileOutputNode.h"
 
+#include "Algo/Count.h"
 #include "Graph/MovieGraphConfig.h"
+#include "Graph/Nodes/MovieGraphGlobalOutputSettingNode.h"
 #include "MoviePipelineUtils.h"
 
 #include UE_INLINE_GENERATED_CPP_BY_NAME(MovieGraphFileOutputNode)
@@ -53,6 +55,151 @@ void UMovieGraphFileOutputNode::DisambiguateFilename(FString& InOutFilenameForma
 	}
 }
 
+void UMovieGraphFileOutputNode::DisambiguateFilename(FString& InOutFilenameFormatString, const UE::MovieGraph::FMovieGraphOutputMergerFrame* InRawFrameData, const UMovieGraphFileOutputNode* InParentNode, const FMovieGraphPassData& InRenderData)
+{
+	constexpr bool bIncludeCDOs = true;
+	constexpr bool bExactMatch = true;
+	const UMovieGraphGlobalOutputSettingNode* OutputSettingNode =
+		InRawFrameData->EvaluatedConfig->GetSettingForBranch<UMovieGraphGlobalOutputSettingNode>(GlobalsPinName, bIncludeCDOs, bExactMatch);
+
+	// Generates the final format string for a given payload, taking into account the file name format from the Output node, any overrides on the
+	// payload, and the output directory from the Global Output Settings node.
+	auto GetFinalFormatString = [InParentNode, OutputSettingNode](const UE::MovieGraph::FMovieGraphSampleState* InPayload) -> FString
+	{
+		if (!InPayload)
+		{
+			return InParentNode->FileNameFormat;
+		}
+
+		// The file name format usually comes from the output node directly, but the payload has a chance to override it.
+		const FString FileNameFormat = !InPayload->FilenameFormatOverride.IsEmpty() ? InPayload->FilenameFormatOverride : InParentNode->FileNameFormat;
+
+		// Generate one string that puts the directory combined with the filename format.
+		FString FileNameFormatString = OutputSettingNode->OutputDirectory.Path / FileNameFormat;
+
+		return FileNameFormatString;
+	};
+
+	const TArray<UE::MovieGraph::FMovieGraphSampleState*>* AllPayloadsForNode = InRawFrameData->NodeInstanceToPayloads.Find(InParentNode);
+	if (!AllPayloadsForNode)
+	{
+		// Shouldn't happen
+		return;
+	}
+	
+	TArray<UE::MovieGraph::FMovieGraphSampleState*> AllPayloadsForRenderer;
+	for (const FMovieGraphPassData& RenderData : InRawFrameData->ImageOutputData)
+	{
+		if (RenderData.Key.RendererName == InRenderData.Key.RendererName)
+		{
+			if (UE::MovieGraph::FMovieGraphSampleState* Payload = RenderData.Value->GetPayload<UE::MovieGraph::FMovieGraphSampleState>())
+			{
+				AllPayloadsForRenderer.Add(Payload);
+			}
+		}
+	}
+
+	const UE::MovieGraph::FMovieGraphSampleState* ThisPayload = nullptr;
+	if (const TUniquePtr<FImagePixelData>* ImagePixelData = InRawFrameData->ImageOutputData.Find(InRenderData.Key))
+	{
+		ThisPayload = (*ImagePixelData)->GetPayload<UE::MovieGraph::FMovieGraphSampleState>();
+	}
+
+	TArray<FString> AllFilePaths;
+	const FString ThisFinalFormatString = GetFinalFormatString(ThisPayload);
+	for (const UE::MovieGraph::FMovieGraphSampleState* NodePayload : *AllPayloadsForNode)
+	{
+		AllFilePaths.Add(GetFinalFormatString(NodePayload));
+	}
+	
+	const FName NodeName = InParentNode->GetFName();
+	const UE::MovieGraph::FMovieGraphRenderDataValidationInfo ValidationInfo = InRawFrameData->GetValidationInfo(InRenderData.Key);
+
+	// Since there can only be one layer per branch, we restrain layer/branch validation to multi-branch graphs.
+	if (ValidationInfo.BranchCount > 1)
+	{
+		// We can run into the scenario where the users have given layers the same name, so layer_name token won't help differentiate.
+		// To resolve this, we look to see if there's multiple branches with the same layer name, and if so we force the branch name into the token too.
+		if (ValidationInfo.LayerCount < ValidationInfo.BranchCount)
+		{
+			UE::MoviePipeline::ConformOutputFormatStringToken(InOutFilenameFormatString, TEXT("{branch_name}"), NodeName, InRenderData.Key.RootBranchName);
+		}
+		else
+		{
+			// Otherwise, we separate each branch by its unique layer name.
+			UE::MoviePipeline::ConformOutputFormatStringToken(InOutFilenameFormatString, TEXT("{layer_name}"), NodeName, InRenderData.Key.RootBranchName);
+		}
+	}
+
+	// We may add the renderer name token if multiple (non-composited) renderers are present on the active branch (eg, in the case of optional PPMs).
+	if (ValidationInfo.ActiveBranchRendererCount > 1)
+	{
+		int32 NumBeautyRenders = 0;
+		for (const UE::MovieGraph::FMovieGraphSampleState* PayloadInstance : *AllPayloadsForNode)
+		{
+			if (PayloadInstance->bIsBeautyPass)
+			{
+				++NumBeautyRenders;
+			}
+		}
+
+		// Determines if the given beauty pass's payload needs the {renderer_name} token.
+		auto BeautyPassNeedsRendererName = [&GetFinalFormatString, &AllFilePaths, &NumBeautyRenders](const UE::MovieGraph::FMovieGraphSampleState* InPayload)
+		{
+			const FString BeautyPassFinalFormatString = GetFinalFormatString(InPayload);
+			const int32 NumPathsWithSameFormatString = Algo::Count(AllFilePaths, BeautyPassFinalFormatString);
+
+			return (NumBeautyRenders > 1) && (NumPathsWithSameFormatString > 1);
+		};
+
+		// * If there is only one "beauty" render in the branch, then don't force on {renderer_name}. Beauty renders should be given priority for
+		//   "opting out" of this token if it's not already specified in the format string.
+		// * If there's more than one "beauty" render in the branch, then {renderer_name} is forced on the beauty passes to disambiguate.
+		// * For subresource passes, generally the {renderer_sub_name} token is enough to disambiguate. However, if its parent beauty pass needs
+		//   {renderer_name}, then the subresource also needs {renderer_name}. 
+		bool bAddRendererName = false;
+		if (ThisPayload && ThisPayload->bIsBeautyPass && BeautyPassNeedsRendererName(ThisPayload))
+		{
+			bAddRendererName = true;
+		}
+		else if (ThisPayload && !ThisPayload->bIsBeautyPass)
+		{
+			// Determine if the parent beauty pass needs {renderer_name}. If so, the subresource also needs it.
+			for (const UE::MovieGraph::FMovieGraphSampleState* RendererPayload : AllPayloadsForRenderer)
+			{
+				if (RendererPayload->bIsBeautyPass && BeautyPassNeedsRendererName(RendererPayload))
+				{
+					bAddRendererName = true;
+					break;
+				}
+			}
+		}
+
+		if (bAddRendererName)
+		{
+			UE::MoviePipeline::ConformOutputFormatStringToken(InOutFilenameFormatString, TEXT("{renderer_name}"), NodeName, InRenderData.Key.RootBranchName);
+		}
+	}
+
+	// We only add the subresource token if a (non-composited) renderer on the active branch is producing more than one subresource (eg, in the case of optional PPMs).
+	if (ValidationInfo.ActiveRendererSubresourceCount > 1)
+	{
+		const bool bIsBeautyPass = ThisPayload && ThisPayload->bIsBeautyPass;
+
+		// However, don't force this token on for beauty passes.
+		if (!bIsBeautyPass)
+		{
+			UE::MoviePipeline::ConformOutputFormatStringToken(InOutFilenameFormatString, TEXT("{renderer_sub_name}"), NodeName, InRenderData.Key.RootBranchName);
+		}
+	}
+
+	// We only add the camera name token if there are more than one cameras for the active branch
+	if (ValidationInfo.ActiveCameraCount > 1)
+	{
+		UE::MoviePipeline::ConformOutputFormatStringToken(InOutFilenameFormatString, TEXT("{camera_name}"), NodeName, InRenderData.Key.RootBranchName);
+	}
+}
+
 TArray<FMovieGraphPassData> UMovieGraphFileOutputNode::GetCompositedPasses(UE::MovieGraph::FMovieGraphOutputMergerFrame* InRawFrameData)
 {
 	// Gather the passes that need to be composited
diff --git a/Engine/Plugins/MovieScene/MovieRenderPipeline/Source/MovieRenderPipelineCore/Private/Graph/Nodes/MovieGraphVideoOutputNode.cpp b/Engine/Plugins/MovieScene/MovieRenderPipeline/Source/MovieRenderPipelineCore/Private/Graph/Nodes/MovieGraphVideoOutputNode.cpp
--- a/Engine/Plugins/MovieScene/MovieRenderPipeline/Source/MovieRenderPipelineCore/Private/Graph/Nodes/MovieGraphVideoOutputNode.cpp
+++ b/Engine/Plugins/MovieScene/MovieRenderPipeline/Source/MovieRenderPipelineCore/Private/Graph/Nodes/MovieGraphVideoOutputNode.cpp
@@ -54,6 +54,10 @@ void UMovieGraphVideoOutputNode::OnAllShotFramesSubmittedImpl(UMovieGraphPipelin
 
 void UMovieGraphVideoOutputNode::OnReceiveImageDataImpl(UMovieGraphPipeline* InPipeline, UE::MovieGraph::FMovieGraphOutputMergerFrame* InRawFrameData, const TSet<FMovieGraphRenderDataIdentifier>& InMask)
 {
+	// ------------------------
+	// This method is called on the CDO!
+	// ------------------------
+
 	// Don't continue processing if there was an error
 	if (bHasError)
 	{
@@ -65,10 +69,19 @@ void UMovieGraphVideoOutputNode::OnReceiveImageDataImpl(UMovieGraphPipeline* InP
 	// Gather the composited passes (eg, Burn Ins and Widget Renderer) 
 	TArray<FMovieGraphPassData> CompositedPasses = GetCompositedPasses(InRawFrameData);
 
+	// Do a first pass to determine which layers/payloads should actually be written by this output node. This needs to be done upfront so we can
+	// provide a list of payloads output to this node to the filename disambiguation logic.
+	TMap<const UMovieGraphFileOutputNode*, TArray<UE::MovieGraph::FMovieGraphSampleState*>> PayloadsToWrite;
 	for (FMovieGraphPassData& RenderPassData : InRawFrameData->ImageOutputData)
 	{
 		FImagePixelData* RawRenderPassData = RenderPassData.Value.Get();
-		const UE::MovieGraph::FMovieGraphSampleState* Payload = RawRenderPassData->GetPayload<UE::MovieGraph::FMovieGraphSampleState>();
+		UE::MovieGraph::FMovieGraphSampleState* Payload = RawRenderPassData->GetPayload<UE::MovieGraph::FMovieGraphSampleState>();
+
+		constexpr bool bIncludeCDOs = false;
+		constexpr bool bExactMatch = true;
+		const UMovieGraphVideoOutputNode* ParentNode = Cast<UMovieGraphVideoOutputNode>(
+			InRawFrameData->EvaluatedConfig->GetSettingForBranch(GetClass(), RenderPassData.Key.RootBranchName, bIncludeCDOs, bExactMatch));
+		checkf(ParentNode, TEXT("Video output should not exist without a parent node in the graph."));
 
 		// This payload may have opted out of being written by this output type.
 		if (!Payload->OutputTypeRestrictions.IsEmpty() && !Payload->OutputTypeRestrictions.Contains(FSoftClassPath(GetClass())))
@@ -91,13 +104,42 @@ void UMovieGraphVideoOutputNode::OnReceiveImageDataImpl(UMovieGraphPipeline* InP
 			continue;
 		}
 
+		PayloadsToWrite.FindOrAdd(ParentNode).Add(Payload);
+	}
+	
+	InRawFrameData->NodeInstanceToPayloads.Append(PayloadsToWrite);
+
+	for (FMovieGraphPassData& RenderPassData : InRawFrameData->ImageOutputData)
+	{
+		constexpr bool bIncludeCDOs = false;
+		constexpr bool bExactMatch = true;
+		const UMovieGraphVideoOutputNode* ParentNode = Cast<UMovieGraphVideoOutputNode>(
+			InRawFrameData->EvaluatedConfig->GetSettingForBranch(GetClass(), RenderPassData.Key.RootBranchName, bIncludeCDOs, bExactMatch));
+		checkf(ParentNode, TEXT("Video output should not exist without a parent node in the graph."));
+
+		TArray<UE::MovieGraph::FMovieGraphSampleState*>* EligiblePayloads = PayloadsToWrite.Find(ParentNode);
+		if (!EligiblePayloads)
+		{
+			continue;
+		}
+
+		// Only write payloads that have been vetted by the prior loop and placed in NodeInstanceToPayloads.
+		FImagePixelData* RawRenderPassData = RenderPassData.Value.Get();
+		UE::MovieGraph::FMovieGraphSampleState* Payload = RawRenderPassData->GetPayload<UE::MovieGraph::FMovieGraphSampleState>();
+		if (!EligiblePayloads->Contains(Payload))
+		{
+			continue;
+		}
+
 		// GetOrCreateOutputWriter() can set bHasError, so even if a writer was created, don't continue if there was an error initializing it
 		const FMovieGraphCodecWriterWithPromise* OutputWriter = GetOrCreateOutputWriter(InPipeline, InRawFrameData, RenderPassData, CompositedPasses);
 		if (!OutputWriter || bHasError)
 		{
+			// This payload is no longer being written, so remove it from NodeInstanceToPayloads
+			InRawFrameData->NodeInstanceToPayloads.Find(ParentNode)->Remove(Payload);
 			continue;
 		}
-
+		
 		// We can have multiple shots in one file, so we find/add the shot in case this is a new shot, and the counter will start over at zero.
 		MovieRenderGraph::IVideoCodecWriter::FLightweightSourceData& LightweightData = OutputWriter->CodecWriter->LightweightSourceData.FindOrAdd(Payload->TraversalContext.ShotIndex);
 		LightweightData.SubmittedFrameCount++;
@@ -266,7 +308,7 @@ void UMovieGraphVideoOutputNode::GetOutputFilePaths(const UMovieGraphPipeline* I
 	UMovieGraphGlobalOutputSettingNode* OutputSettingsNode =
 		EvaluatedConfig->GetSettingForBranch<UMovieGraphGlobalOutputSettingNode>(GlobalsPinName, bIncludeCDOs, bExactMatch);
 
-	UMovieGraphVideoOutputNode* EvaluatedNode = Cast<UMovieGraphVideoOutputNode>(
+	const UMovieGraphVideoOutputNode* EvaluatedNode = Cast<UMovieGraphVideoOutputNode>(
 		EvaluatedConfig->GetSettingForBranch(GetClass(), InRenderPassData.Key.RootBranchName, bIncludeCDOs, bExactMatch));
 
 	const UE::MovieGraph::FMovieGraphSampleState* Payload = InRenderPassData.Value->GetPayload<UE::MovieGraph::FMovieGraphSampleState>();
@@ -286,14 +328,14 @@ void UMovieGraphVideoOutputNode::GetOutputFilePaths(const UMovieGraphPipeline* I
 	
 	FMovieGraphResolveArgs FinalFormatArgs;
 	
-	// Insert tokens like {layer_name} as appropriate to make sure outputs don't clash with each other.
-	DisambiguateFilename(FileNameFormatString, InRawFrameData, EvaluatedNode->GetFName(), InRenderPassData);
-
 	// Strip any frame number tags so we don't get one video file per frame.
 	UE::MoviePipeline::RemoveFrameNumberFormatStrings(FileNameFormatString, true);
 
 	const FString OutputDirectory = OutputSettingsNode->OutputDirectory.Path;
-	const FString FullFilepathFormatString = OutputDirectory / FileNameFormatString;
+	FString FullFilepathFormatString = OutputDirectory / FileNameFormatString;
+
+	// Insert tokens like {layer_name} as appropriate to make sure outputs don't clash with each other.
+	DisambiguateFilename(FullFilepathFormatString, InRawFrameData, EvaluatedNode, InRenderPassData);
 
 	// This is a bit crummy but we don't have a good way to implement bOverwriteFiles = False for video files. When you don't have the
 	// override file flag set, we need to increment the number by 1, and write to that. This works fine for image sequences as image
diff --git a/Engine/Plugins/MovieScene/MovieRenderPipeline/Source/MovieRenderPipelineCore/Public/Graph/MovieGraphDataTypes.h b/Engine/Plugins/MovieScene/MovieRenderPipeline/Source/MovieRenderPipelineCore/Public/Graph/MovieGraphDataTypes.h
--- a/Engine/Plugins/MovieScene/MovieRenderPipeline/Source/MovieRenderPipelineCore/Public/Graph/MovieGraphDataTypes.h
+++ b/Engine/Plugins/MovieScene/MovieRenderPipeline/Source/MovieRenderPipelineCore/Public/Graph/MovieGraphDataTypes.h
@@ -18,8 +18,9 @@
 #define UE_API MOVIERENDERPIPELINECORE_API
 
 // Forward Declares
-class UMovieGraphPipeline;
 class UMovieGraphEvaluatedConfig;
+class UMovieGraphFileOutputNode;
+class UMovieGraphPipeline;
 class UMovieGraphTimeStepBase;
 class UMovieGraphRendererBase;
 class UMovieGraphTimeRangeBuilderBase;
@@ -566,6 +567,12 @@ namespace UE::MovieGraph
 		/** Additional metadata to be added to the output (if supported by the output container). */
 		TMap<FString, FString> FileMetadata;
 
+		/**
+		 * Specifies which payloads will be written from a specific output node instance.
+		 * NOTE: This is only updated immediately prior to filename disambiguation and should not be used outside of that.
+		 */
+		TMap<const UMovieGraphFileOutputNode*, TArray<UE::MovieGraph::FMovieGraphSampleState*>> NodeInstanceToPayloads;
+
 		/** Get filename token validation info from the (expected) render passes. */
 		MOVIERENDERPIPELINECORE_API FMovieGraphRenderDataValidationInfo GetValidationInfo(const FMovieGraphRenderDataIdentifier& InRenderID, bool bInDiscardCompositedRenders = true) const;
 	};
diff --git a/Engine/Plugins/MovieScene/MovieRenderPipeline/Source/MovieRenderPipelineCore/Public/Graph/Nodes/MovieGraphFileOutputNode.h b/Engine/Plugins/MovieScene/MovieRenderPipeline/Source/MovieRenderPipelineCore/Public/Graph/Nodes/MovieGraphFileOutputNode.h
--- a/Engine/Plugins/MovieScene/MovieRenderPipeline/Source/MovieRenderPipelineCore/Public/Graph/Nodes/MovieGraphFileOutputNode.h
+++ b/Engine/Plugins/MovieScene/MovieRenderPipeline/Source/MovieRenderPipelineCore/Public/Graph/Nodes/MovieGraphFileOutputNode.h
@@ -98,8 +98,12 @@ protected:
 	static UE_API int32 GetNumFileOutputNodes(const UMovieGraphEvaluatedConfig& InEvaluatedConfig, const FName& InBranchName);
 
 	/** Updates InOutFilenameFormatString to include tokens (like {layer_name}) that will disambiguate the output if multiple outputs are attempting to write to the same file. */
+	UE_DEPRECATED(5.7, "Please use the version of DisambiguateFilename() that takes a UMovieGraphFileOutputNode instance. The new version supports enhanced disambiguation rules.")
 	static UE_API void DisambiguateFilename(FString& InOutFilenameFormatString, const UE::MovieGraph::FMovieGraphOutputMergerFrame* InRawFrameData, const FName& InNodeName, const FMovieGraphPassData& InRenderData);
 
+	/** Updates InOutFilenameFormatString to include tokens (like {layer_name}) that will disambiguate the output if multiple outputs are attempting to write to the same file. */
+	static UE_API void DisambiguateFilename(FString& InOutFilenameFormatString, const UE::MovieGraph::FMovieGraphOutputMergerFrame* InRawFrameData, const UMovieGraphFileOutputNode* InParentNode, const FMovieGraphPassData& InRenderData);
+
 	/** Convenience function to get the list of active composite passes from render data. */
 	static UE_API TArray<FMovieGraphPassData> GetCompositedPasses(UE::MovieGraph::FMovieGraphOutputMergerFrame* InRawFrameData);
 };
diff --git a/Engine/Plugins/MovieScene/MovieRenderPipeline/Source/MovieRenderPipelineRenderPasses/Private/Graph/Nodes/MovieGraphDeferredPassNode.cpp b/Engine/Plugins/MovieScene/MovieRenderPipeline/Source/MovieRenderPipelineRenderPasses/Private/Graph/Nodes/MovieGraphDeferredPassNode.cpp
--- a/Engine/Plugins/MovieScene/MovieRenderPipeline/Source/MovieRenderPipelineRenderPasses/Private/Graph/Nodes/MovieGraphDeferredPassNode.cpp
+++ b/Engine/Plugins/MovieScene/MovieRenderPipeline/Source/MovieRenderPipelineRenderPasses/Private/Graph/Nodes/MovieGraphDeferredPassNode.cpp
@@ -21,6 +21,7 @@ UMovieGraphDeferredRenderPassNode::UMovieGraphDeferredRenderPassNode()
 	, bAllowOCIO(true)
 	, ViewModeIndex(VMI_Lit)
 	, bIncludeBeautyRenderInOutput(true)
+	, PPMFileNameFormat(TEXT("{sequence_name}.{layer_name}.{renderer_sub_name}.{frame_number}"))
 	, bEnableHighResolutionTiling(false)
 	, TileCount(1)
 	, OverlapPercentage(0.f)
@@ -121,6 +122,11 @@ TArray<FMoviePipelinePostProcessPass> UMovieGraphDeferredRenderPassNode::GetAddi
 	return AdditionalPostProcessMaterials;
 }
 
+FString UMovieGraphDeferredRenderPassNode::GetPPMFileNameFormatString() const
+{
+	return PPMFileNameFormat;
+}
+
 int32 UMovieGraphDeferredRenderPassNode::GetNumSpatialSamples() const
 {
 	return SpatialSampleCount;
diff --git a/Engine/Plugins/MovieScene/MovieRenderPipeline/Source/MovieRenderPipelineRenderPasses/Private/Graph/Nodes/MovieGraphPathTracerPassNode.cpp b/Engine/Plugins/MovieScene/MovieRenderPipeline/Source/MovieRenderPipelineRenderPasses/Private/Graph/Nodes/MovieGraphPathTracerPassNode.cpp
--- a/Engine/Plugins/MovieScene/MovieRenderPipeline/Source/MovieRenderPipelineRenderPasses/Private/Graph/Nodes/MovieGraphPathTracerPassNode.cpp
+++ b/Engine/Plugins/MovieScene/MovieRenderPipeline/Source/MovieRenderPipelineRenderPasses/Private/Graph/Nodes/MovieGraphPathTracerPassNode.cpp
@@ -28,6 +28,7 @@ UMovieGraphPathTracerRenderPassNode::UMovieGraphPathTracerRenderPassNode()
 	, bDisableToneCurve(false)
 	, bAllowOCIO(true)
 	, bIncludeBeautyRenderInOutput(true)
+	, PPMFileNameFormat(TEXT("{sequence_name}.{layer_name}.{renderer_sub_name}.{frame_number}"))
 	, TileCount(1)
 {
 	RendererName = TEXT("PathTraced");
@@ -210,6 +211,11 @@ TArray<FMoviePipelinePostProcessPass> UMovieGraphPathTracerRenderPassNode::GetAd
 	return AdditionalPostProcessMaterials;
 }
 
+FString UMovieGraphPathTracerRenderPassNode::GetPPMFileNameFormatString() const
+{
+	return PPMFileNameFormat;
+}
+
 int32 UMovieGraphPathTracerRenderPassNode::GetNumSpatialSamples() const
 {
 	return SpatialSampleCount;
diff --git a/Engine/Plugins/MovieScene/MovieRenderPipeline/Source/MovieRenderPipelineRenderPasses/Private/Graph/Renderers/MovieGraphDeferredPass.cpp b/Engine/Plugins/MovieScene/MovieRenderPipeline/Source/MovieRenderPipelineRenderPasses/Private/Graph/Renderers/MovieGraphDeferredPass.cpp
--- a/Engine/Plugins/MovieScene/MovieRenderPipeline/Source/MovieRenderPipelineRenderPasses/Private/Graph/Renderers/MovieGraphDeferredPass.cpp
+++ b/Engine/Plugins/MovieScene/MovieRenderPipeline/Source/MovieRenderPipelineRenderPasses/Private/Graph/Renderers/MovieGraphDeferredPass.cpp
@@ -543,6 +543,9 @@ void FMovieGraphDeferredPass::Render(const FMovieGraphTraversalContext& InFrameT
 						// PPM passes are never the beauty pass.
 						PassSampleState.bIsBeautyPass = false;
 
+						// PPM passes may override the file name format that they use.
+						PassSampleState.FilenameFormatOverride = ParentNode->GetPPMFileNameFormatString();
+
 						// Give a lower priority to materials so they show up after the main pass in multi-layer exrs.
 						PassSampleState.CompositingSortOrder = SampleState.CompositingSortOrder + 1;
 						BufferPipe->AddEndpoint(MakeForwardingEndpoint(PassSampleState, InTimeData));
diff --git a/Engine/Plugins/MovieScene/MovieRenderPipeline/Source/MovieRenderPipelineRenderPasses/Private/MovieGraphImageSequenceOutputNode.cpp b/Engine/Plugins/MovieScene/MovieRenderPipeline/Source/MovieRenderPipelineRenderPasses/Private/MovieGraphImageSequenceOutputNode.cpp
--- a/Engine/Plugins/MovieScene/MovieRenderPipeline/Source/MovieRenderPipelineRenderPasses/Private/MovieGraphImageSequenceOutputNode.cpp
+++ b/Engine/Plugins/MovieScene/MovieRenderPipeline/Source/MovieRenderPipelineRenderPasses/Private/MovieGraphImageSequenceOutputNode.cpp
@@ -121,8 +121,7 @@ FString UMovieGraphImageSequenceOutputNode::CreateFileName(
 	UE::MovieGraph::FMovieGraphOutputMergerFrame* InRawFrameData,
 	const UMovieGraphImageSequenceOutputNode* InParentNode,
 	const UMovieGraphPipeline* InPipeline,
-	const TPair<FMovieGraphRenderDataIdentifier,
-	TUniquePtr<FImagePixelData>>& InRenderData,
+	const TPair<FMovieGraphRenderDataIdentifier, TUniquePtr<FImagePixelData>>& InRenderData,
 	const EImageFormat InImageFormat,
 	FMovieGraphResolveArgs& OutMergedFormatArgs,
 	FString& OutFrameTemplatedFileName) const
@@ -151,7 +150,7 @@ FString UMovieGraphImageSequenceOutputNode::CreateFileName(
 	FString FileNameFormatString = OutputSettingNode->OutputDirectory.Path / FinalFileNameFormat;
 
 	// Insert tokens like {layer_name} as appropriate to make sure outputs don't clash with each other.
-	DisambiguateFilename(FileNameFormatString, InRawFrameData, InParentNode->GetFName(), InRenderData);
+	DisambiguateFilename(FileNameFormatString, InRawFrameData, InParentNode, InRenderData);
 
 	// Previous method is preserved for output frame number validation.
 	constexpr bool bIncludeRenderPass = false;
@@ -212,6 +211,10 @@ FString UMovieGraphImageSequenceOutputNode::GetFrameTemplatedFileName(
 
 void UMovieGraphImageSequenceOutputNode::OnReceiveImageDataImpl(UMovieGraphPipeline* InPipeline, UE::MovieGraph::FMovieGraphOutputMergerFrame* InRawFrameData, const TSet<FMovieGraphRenderDataIdentifier>& InMask)
 {
+	// ------------------------
+	// This method is called on the CDO!
+	// ------------------------
+
 	check(InRawFrameData);
 
 	TArray<TPair<FMovieGraphRenderDataIdentifier, TUniquePtr<FImagePixelData>>> CompositedPasses = GetCompositedPasses(InRawFrameData);
@@ -228,10 +231,19 @@ void UMovieGraphImageSequenceOutputNode::OnReceiveImageDataImpl(UMovieGraphPipel
 	// other with dependencies, but I'm not sure that's going to scale to the potentialy high data volume going wide MRQ will eventually
 	// need.
 
-	// The base ImageSequenceOutputNode doesn't support any multilayer formats, so we write out each render pass separately.
+	// Do a first pass to determine which layers/payloads should actually be written by this output node. This needs to be done upfront so we can
+	// provide a list of payloads output to this node to the filename disambiguation logic.
+	TMap<const UMovieGraphFileOutputNode*, TArray<UE::MovieGraph::FMovieGraphSampleState*>> PayloadsToWrite;
 	for (TPair<FMovieGraphRenderDataIdentifier, TUniquePtr<FImagePixelData>>& RenderData : InRawFrameData->ImageOutputData)
 	{
+		checkf(RenderData.Value.IsValid(), TEXT("Unexpected empty image data: incorrectly moved or its production failed?"));
 		UE::MovieGraph::FMovieGraphSampleState* Payload = RenderData.Value->GetPayload<UE::MovieGraph::FMovieGraphSampleState>();
+		
+		constexpr bool bIncludeCDOs = false;
+		constexpr bool bExactMatch = true;
+		const UMovieGraphImageSequenceOutputNode* ParentNode = Cast<UMovieGraphImageSequenceOutputNode>(
+			InRawFrameData->EvaluatedConfig->GetSettingForBranch(GetClass(), RenderData.Key.RootBranchName, bIncludeCDOs, bExactMatch));
+		checkf(ParentNode, TEXT("Image sequence output should not exist without a parent node in the graph."));
 
 		// This payload may have opted out of being written by this output type.
 		if (!Payload->OutputTypeRestrictions.IsEmpty() && !Payload->OutputTypeRestrictions.Contains(FSoftClassPath(GetClass())))
@@ -254,28 +266,43 @@ void UMovieGraphImageSequenceOutputNode::OnReceiveImageDataImpl(UMovieGraphPipel
 			continue;
 		}
 
-		checkf(RenderData.Value.IsValid(), TEXT("Unexpected empty image data: incorrectly moved or its production failed?"));
-
-		// ToDo: Certain images may require transparency, at which point
-		// we write out a .png instead of a .jpeg.
-		EImageFormat PreferredOutputFormat = OutputFormat;
+		PayloadsToWrite.FindOrAdd(ParentNode).Add(Payload);
+	}
+	
+	InRawFrameData->NodeInstanceToPayloads.Append(PayloadsToWrite);
 
+	// The base ImageSequenceOutputNode doesn't support any multilayer formats, so we write out each render pass separately.
+	for (TPair<FMovieGraphRenderDataIdentifier, TUniquePtr<FImagePixelData>>& RenderData : InRawFrameData->ImageOutputData)
+	{
 		constexpr bool bIncludeCDOs = false;
 		constexpr bool bExactMatch = true;
 		const UMovieGraphImageSequenceOutputNode* ParentNode = Cast<UMovieGraphImageSequenceOutputNode>(
 			InRawFrameData->EvaluatedConfig->GetSettingForBranch(GetClass(), RenderData.Key.RootBranchName, bIncludeCDOs, bExactMatch));
 		checkf(ParentNode, TEXT("Image sequence output should not exist without a parent node in the graph."));
+
+		TArray<UE::MovieGraph::FMovieGraphSampleState*>* EligiblePayloads = PayloadsToWrite.Find(ParentNode);
+		if (!EligiblePayloads)
+		{
+			continue;
+		}
+
+		// Only write payloads that have been vetted by the prior loop and placed in NodeInstanceToPayloads.
+		UE::MovieGraph::FMovieGraphSampleState* Payload = RenderData.Value->GetPayload<UE::MovieGraph::FMovieGraphSampleState>();
+		if (!EligiblePayloads->Contains(Payload))
+		{
+			continue;
+		}
 		
 		FMovieGraphResolveArgs FinalResolvedKVPs;
 		FString FrameTemplatedFileName;
-		FString FileName = CreateFileName(InRawFrameData, ParentNode, InPipeline, RenderData, PreferredOutputFormat, FinalResolvedKVPs, FrameTemplatedFileName);
+		FString FileName = CreateFileName(InRawFrameData, ParentNode, InPipeline, RenderData, OutputFormat, FinalResolvedKVPs, FrameTemplatedFileName);
 		if (!ensureMsgf(!FileName.IsEmpty(), TEXT("Unexpected empty file name, skipping frame.")))
 		{
 			continue;
 		}
 
 		TUniquePtr<FImageWriteTask> TileImageTask = MakeUnique<FImageWriteTask>();
-		TileImageTask->Format = PreferredOutputFormat;
+		TileImageTask->Format = OutputFormat;
 		TileImageTask->CompressionQuality = 100;
 		TileImageTask->Filename = FileName;
 
diff --git a/Engine/Plugins/MovieScene/MovieRenderPipeline/Source/MovieRenderPipelineRenderPasses/Public/Graph/Nodes/MovieGraphDeferredPassNode.h b/Engine/Plugins/MovieScene/MovieRenderPipeline/Source/MovieRenderPipelineRenderPasses/Public/Graph/Nodes/MovieGraphDeferredPassNode.h
--- a/Engine/Plugins/MovieScene/MovieRenderPipeline/Source/MovieRenderPipelineRenderPasses/Public/Graph/Nodes/MovieGraphDeferredPassNode.h
+++ b/Engine/Plugins/MovieScene/MovieRenderPipeline/Source/MovieRenderPipelineRenderPasses/Public/Graph/Nodes/MovieGraphDeferredPassNode.h
@@ -32,6 +32,7 @@ protected:
 	UE_API virtual TUniquePtr<UE::MovieGraph::Rendering::FMovieGraphImagePassBase> CreateInstance() const override;
 	UE_API virtual bool GetWriteAllSamples() const override;
 	UE_API virtual TArray<FMoviePipelinePostProcessPass> GetAdditionalPostProcessMaterials() const override;
+	UE_API virtual FString GetPPMFileNameFormatString() const override;
 	UE_API virtual int32 GetNumSpatialSamples() const override;
 	UE_API virtual bool GetDisableToneCurve() const override;
 	UE_API virtual bool GetAllowOCIO() const override;
@@ -70,6 +71,9 @@ public:
 
 	UPROPERTY(EditAnywhere, BlueprintReadWrite, Category = Overrides, meta = (InlineEditConditionToggle))
 	uint8 bOverride_bIncludeBeautyRenderInOutput : 1;
+
+	UPROPERTY(EditAnywhere, BlueprintReadWrite, Category = Overrides, meta = (InlineEditConditionToggle))
+	uint8 bOverride_PPMFileNameFormat : 1;
 	
 	UPROPERTY(EditAnywhere, BlueprintReadWrite, Category = Overrides, meta = (InlineEditConditionToggle))
 	uint8 bOverride_AdditionalPostProcessMaterials : 1;
@@ -147,6 +151,12 @@ public:
 	UPROPERTY(EditAnywhere, BlueprintReadWrite, Category = "Post Process Materials", meta=(EditCondition="bOverride_bIncludeBeautyRenderInOutput"))
 	bool bIncludeBeautyRenderInOutput;
 
+	/**
+	 * If specified, overrides the output node's file name format, and uses this file name format for Post Process Material passes instead.
+	 */
+	UPROPERTY(EditAnywhere, BlueprintReadWrite, Category = "Post Process Materials", DisplayName = "PPM File Name Format", meta=(EditCondition="bOverride_PPMFileNameFormat"))
+	FString PPMFileNameFormat;
+
 	/**
 	* An array of additional post-processing materials to run after the frame is rendered. Using this feature may add a notable amount of render time.
 	*/
diff --git a/Engine/Plugins/MovieScene/MovieRenderPipeline/Source/MovieRenderPipelineRenderPasses/Public/Graph/Nodes/MovieGraphImagePassBaseNode.h b/Engine/Plugins/MovieScene/MovieRenderPipeline/Source/MovieRenderPipelineRenderPasses/Public/Graph/Nodes/MovieGraphImagePassBaseNode.h
--- a/Engine/Plugins/MovieScene/MovieRenderPipeline/Source/MovieRenderPipelineRenderPasses/Public/Graph/Nodes/MovieGraphImagePassBaseNode.h
+++ b/Engine/Plugins/MovieScene/MovieRenderPipeline/Source/MovieRenderPipelineRenderPasses/Public/Graph/Nodes/MovieGraphImagePassBaseNode.h
@@ -46,6 +46,12 @@ public:
 	/** Get an array of user-added post-process materials for the render */
 	virtual TArray<FMoviePipelinePostProcessPass> GetAdditionalPostProcessMaterials() const { return {}; }
 
+	/**
+	 * Get the post-process material file name format string, if any. If empty, no file name format override will be used for PPMs, and the
+	 * value on the output node will be used.
+	 */
+	virtual FString GetPPMFileNameFormatString() const { return FString(); }
+
 	/** How many spatial samples should be rendered each frame? */
 	virtual int32 GetNumSpatialSamples() const { return 1; }
 
diff --git a/Engine/Plugins/MovieScene/MovieRenderPipeline/Source/MovieRenderPipelineRenderPasses/Public/Graph/Nodes/MovieGraphPathTracerPassNode.h b/Engine/Plugins/MovieScene/MovieRenderPipeline/Source/MovieRenderPipelineRenderPasses/Public/Graph/Nodes/MovieGraphPathTracerPassNode.h
--- a/Engine/Plugins/MovieScene/MovieRenderPipeline/Source/MovieRenderPipelineRenderPasses/Public/Graph/Nodes/MovieGraphPathTracerPassNode.h
+++ b/Engine/Plugins/MovieScene/MovieRenderPipeline/Source/MovieRenderPipelineRenderPasses/Public/Graph/Nodes/MovieGraphPathTracerPassNode.h
@@ -50,6 +50,7 @@ public:
 	// UMovieGraphImagePassBaseNode Interface
 	UE_API virtual bool GetWriteAllSamples() const override;
 	UE_API virtual TArray<FMoviePipelinePostProcessPass> GetAdditionalPostProcessMaterials() const override;
+	UE_API virtual FString GetPPMFileNameFormatString() const override;
 	UE_API virtual int32 GetNumSpatialSamples() const override;
 	UE_API virtual int32 GetNumSpatialSamplesDuringWarmUp() const override;
 	UE_API virtual int32 GetSeedOffset() const override;
@@ -135,6 +136,9 @@ public:
 	UPROPERTY(EditAnywhere, BlueprintReadWrite, Category = Overrides, meta = (InlineEditConditionToggle))
 	uint8 bOverride_bIncludeBeautyRenderInOutput : 1;
 
+	UPROPERTY(EditAnywhere, BlueprintReadWrite, Category = Overrides, meta = (InlineEditConditionToggle))
+	uint8 bOverride_PPMFileNameFormat : 1;
+
 	UPROPERTY(EditAnywhere, BlueprintReadWrite, Category = Overrides, meta = (InlineEditConditionToggle))
 	uint8 bOverride_AdditionalPostProcessMaterials : 1;
 
@@ -256,6 +260,12 @@ public:
 	UPROPERTY(EditAnywhere, BlueprintReadWrite, Category = "Post Process Materials", meta=(EditCondition="bOverride_bIncludeBeautyRenderInOutput"))
 	bool bIncludeBeautyRenderInOutput;
 
+	/**
+	 * If specified, overrides the output node's file name format, and uses this file name format for Post Process Material passes instead.
+	 */
+	UPROPERTY(EditAnywhere, BlueprintReadWrite, Category = "Post Process Materials", DisplayName = "PPM File Name Format", meta=(EditCondition="bOverride_PPMFileNameFormat"))
+	FString PPMFileNameFormat;
+
 	/**
 	* An array of additional post-processing materials to run after the frame is rendered. Using this feature may add a notable amount of render time.
 	*/
