MoviePipeline: Fixed filename format disambiguation for EXRs. EXRs were left out of change 45466944.

diff --git a/Engine/Plugins/MovieScene/MovieRenderPipeline/Source/MovieRenderPipelineRenderPasses/Private/MovieGraphImageSequenceOutputNode.cpp b/Engine/Plugins/MovieScene/MovieRenderPipeline/Source/MovieRenderPipelineRenderPasses/Private/MovieGraphImageSequenceOutputNode.cpp
--- a/Engine/Plugins/MovieScene/MovieRenderPipeline/Source/MovieRenderPipelineRenderPasses/Private/MovieGraphImageSequenceOutputNode.cpp
+++ b/Engine/Plugins/MovieScene/MovieRenderPipeline/Source/MovieRenderPipelineRenderPasses/Private/MovieGraphImageSequenceOutputNode.cpp
@@ -515,10 +515,20 @@ void UMovieGraphImageSequenceOutputNode_EXR::OnReceiveImageDataImpl(UMovieGraphP
 
 	TArray<TPair<FMovieGraphRenderDataIdentifier, TUniquePtr<FImagePixelData>>> CompositedPasses = GetCompositedPasses(InRawFrameData);
 
+	// Do a first pass to determine which layers/payloads should actually be written by this output node. This needs to be done upfront so we can
+	// provide a list of payloads output by this node to the filename disambiguation logic.
+	TMap<const UMovieGraphFileOutputNode*, TArray<UE::MovieGraph::FMovieGraphSampleState*>> PayloadsToWrite;
 	for (TPair<FMovieGraphRenderDataIdentifier, TUniquePtr<FImagePixelData>>& RenderData : InRawFrameData->ImageOutputData)
 	{
+		checkf(RenderData.Value.IsValid(), TEXT("Unexpected empty image data: incorrectly moved or its production failed?"));
 		UE::MovieGraph::FMovieGraphSampleState* Payload = RenderData.Value->GetPayload<UE::MovieGraph::FMovieGraphSampleState>();
 
+		constexpr bool bIncludeCDOs = false;
+		constexpr bool bExactMatch = true;
+		const UMovieGraphImageSequenceOutputNode* ParentNode = Cast<UMovieGraphImageSequenceOutputNode>(
+			InRawFrameData->EvaluatedConfig->GetSettingForBranch(GetClass(), RenderData.Key.RootBranchName, bIncludeCDOs, bExactMatch));
+		checkf(ParentNode, TEXT("Image sequence output should not exist without a parent node in the graph."));
+
 		// This payload may have opted out of being written by this output type.
 		if (!Payload->OutputTypeRestrictions.IsEmpty() && !Payload->OutputTypeRestrictions.Contains(FSoftClassPath(GetClass())))
 		{
@@ -540,14 +550,32 @@ void UMovieGraphImageSequenceOutputNode_EXR::OnReceiveImageDataImpl(UMovieGraphP
 			continue;
 		}
 
-		checkf(RenderData.Value.IsValid(), TEXT("Unexpected empty image data: incorrectly moved or its production failed?"));
+		PayloadsToWrite.FindOrAdd(ParentNode).Add(Payload);
+	}
 
+	InRawFrameData->NodeInstanceToPayloads.Append(PayloadsToWrite);
+
+	for (TPair<FMovieGraphRenderDataIdentifier, TUniquePtr<FImagePixelData>>& RenderData : InRawFrameData->ImageOutputData)
+	{
 		constexpr bool bIncludeCDOs = false;
 		constexpr bool bExactMatch = true;
 		const UMovieGraphImageSequenceOutputNode_EXR* ParentNode = InRawFrameData->EvaluatedConfig->GetSettingForBranch<UMovieGraphImageSequenceOutputNode_EXR>(
 			RenderData.Key.RootBranchName, bIncludeCDOs, bExactMatch);
 		checkf(ParentNode, TEXT("Single-layer EXR should not exist without a parent node in the graph."));
 
+		TArray<UE::MovieGraph::FMovieGraphSampleState*>* EligiblePayloads = PayloadsToWrite.Find(ParentNode);
+		if (!EligiblePayloads)
+		{
+			continue;
+		}
+
+		// Only write payloads that have been vetted by the prior loop and placed in NodeInstanceToPayloads.
+		UE::MovieGraph::FMovieGraphSampleState* Payload = RenderData.Value->GetPayload<UE::MovieGraph::FMovieGraphSampleState>();
+		if (!EligiblePayloads->Contains(Payload))
+		{
+			continue;
+		}
+
 		FMovieGraphResolveArgs ResolvedFormatArgs;
 		FString FrameTemplatedFileName;
 		FString FileName = CreateFileName(InRawFrameData, ParentNode, InPipeline, RenderData, OutputFormat, ResolvedFormatArgs, FrameTemplatedFileName);
