MovieGraph: Add support for boolean cvar values in Console Variable Preset assets. These were supported by the Console Variable Editor, but would fail to apply in an MRQ/MRG job.

Added support to FDefaultValueHelper for parsing bools, accepted form is "true" "false" "1", "0". Follows standard C-rules for non-zero integers. Tested parsing with the following strings:
True, False, TrUe, tru, TRUE, 1, 2, 0, Zero

And verified that they had the correct outcome for being parsed (false for tru, and Zero)

diff --git a/Engine/Plugins/MovieScene/MovieRenderPipeline/Source/MovieRenderPipelineCore/Private/Graph/MovieGraphCVarManager.cpp b/Engine/Plugins/MovieScene/MovieRenderPipeline/Source/MovieRenderPipelineCore/Private/Graph/MovieGraphCVarManager.cpp
--- a/Engine/Plugins/MovieScene/MovieRenderPipeline/Source/MovieRenderPipelineCore/Private/Graph/MovieGraphCVarManager.cpp
+++ b/Engine/Plugins/MovieScene/MovieRenderPipeline/Source/MovieRenderPipelineCore/Private/Graph/MovieGraphCVarManager.cpp
@@ -1,4 +1,4 @@
-﻿// Copyright Epic Games, Inc. All Rights Reserved.
+// Copyright Epic Games, Inc. All Rights Reserved.
 
 #include "Graph/MovieGraphCVarManager.h"
 
@@ -49,10 +49,15 @@ void FMovieGraphCVarManager::AddPreset(const TScriptInterface<IMovieSceneConsole
 	for (const TTuple<FString, FString>& CVarPair : PresetCVars)
 	{
 		float CVarFloatValue = 0.0f;
+		bool bCvarBoolValue = false;
 		if (FDefaultValueHelper::ParseFloat(CVarPair.Value, CVarFloatValue))
 		{
 			AddCVar(CVarPair.Key, CVarFloatValue);
 		}
+		else if (FDefaultValueHelper::ParseBool(CVarPair.Value, bCvarBoolValue))
+		{
+			AddCVar(CVarPair.Key, bCvarBoolValue ? 1.0f : 0.0f);
+		}
 		else
 		{
 			UE_LOG(LogMovieRenderPipeline, Warning, TEXT("Failed to apply CVar \"%s\" (from preset \"%s\") because value could not be parsed into a float. Ignoring."),
@@ -226,4 +231,4 @@ void FMovieGraphCVarManager::ApplyCVar(IConsoleVariable* InCVar, float InValue)
 	}
 }
 
-} // namespace UE::MovieGraph::Private
\ No newline at end of file
+} // namespace UE::MovieGraph::Private
diff --git a/Engine/Plugins/MovieScene/MovieRenderPipeline/Source/MovieRenderPipelineSettings/Private/MoviePipelineConsoleVariableSetting.cpp b/Engine/Plugins/MovieScene/MovieRenderPipeline/Source/MovieRenderPipelineSettings/Private/MoviePipelineConsoleVariableSetting.cpp
--- a/Engine/Plugins/MovieScene/MovieRenderPipeline/Source/MovieRenderPipelineSettings/Private/MoviePipelineConsoleVariableSetting.cpp
+++ b/Engine/Plugins/MovieScene/MovieRenderPipeline/Source/MovieRenderPipelineSettings/Private/MoviePipelineConsoleVariableSetting.cpp
@@ -367,10 +367,16 @@ void UMoviePipelineConsoleVariableSetting::MergeInPresetConsoleVariables()
 		for (const TTuple<FString, FString>& CVarPair : PresetCVars)
 		{
 			float CVarFloatValue = 0.0f;
+			bool bCvarBoolValue = false;
+			
 			if (FDefaultValueHelper::ParseFloat(CVarPair.Value, CVarFloatValue))
 			{
 				MergedConsoleVariables.Add(FMoviePipelineConsoleVariableEntry(CVarPair.Key, CVarFloatValue));
 			}
+			else if (FDefaultValueHelper::ParseBool(CVarPair.Value, bCvarBoolValue))
+			{
+				MergedConsoleVariables.Add(FMoviePipelineConsoleVariableEntry(CVarPair.Key, bCvarBoolValue ? 1.0f : 0.0f));
+			}
 			else
 			{
 				UE_LOG(LogMovieRenderPipeline, Warning, TEXT("Failed to apply CVar \"%s\" (from preset \"%s\") because value could not be parsed into a float. Ignoring."),
diff --git a/Engine/Source/Runtime/Core/Private/Misc/DefaultValueHelper.cpp b/Engine/Source/Runtime/Core/Private/Misc/DefaultValueHelper.cpp
--- a/Engine/Source/Runtime/Core/Private/Misc/DefaultValueHelper.cpp
+++ b/Engine/Source/Runtime/Core/Private/Misc/DefaultValueHelper.cpp
@@ -619,6 +619,40 @@ bool FDefaultValueHelper::ParseInt(const FString& Source, int32& OutVal)
 	return false;
 }
 
+bool FDefaultValueHelper::ParseBool(const FString& Source, bool& OutVal)
+{
+	const bool bHasWhitespace = HasWhitespaces(Source);
+	const FString NoWhitespace = bHasWhitespace ? RemoveWhitespaces(Source) : FString();
+	const FString& ProperSource = bHasWhitespace ? NoWhitespace : Source;
+	if (ProperSource.IsEmpty())
+	{
+		return false;
+	}
+
+	const FString SourceToLower = ProperSource.ToLower();
+	if (SourceToLower == TEXT("true"))
+	{
+		OutVal = true;
+		return true;
+	}
+	
+	if (SourceToLower == TEXT("false"))
+	{
+		OutVal = false;
+		return true;
+	}
+	
+	// If it's not in the "true/false" form, try to parse it as an integer
+	int32 Base;
+	if (!Source.IsEmpty() && IsStringValidInteger(StartOf(Source), EndOf(Source), Base))
+	{
+		int32 ValAsInteger = FCString::Strtoi(*ProperSource, NULL, Base);
+		OutVal = static_cast<bool>(ValAsInteger);
+		return true;
+	}
+
+	return false;
+}
 
 bool FDefaultValueHelper::ParseInt64(const FString& Source, int64& OutVal)
 {
diff --git a/Engine/Source/Runtime/Core/Public/Misc/DefaultValueHelper.h b/Engine/Source/Runtime/Core/Public/Misc/DefaultValueHelper.h
--- a/Engine/Source/Runtime/Core/Public/Misc/DefaultValueHelper.h
+++ b/Engine/Source/Runtime/Core/Public/Misc/DefaultValueHelper.h
@@ -97,6 +97,17 @@ public:
 	 */
 	static CORE_API bool ParseInt(const FString& Source, int32& OutVal);
 
+	/**
+	 * Converts a string into a bool.
+	 * Accepted form: "true" or "false" or "1" or "0"
+	 *
+	 * @param Source the input string to try to convert
+	 * @param OutVal the output boolean
+	 *
+	 * @return true if the conversion happened, false otherwise
+	 */
+	static CORE_API bool ParseBool(const FString& Source, bool& OutVal);
+
 	/**
 	 * Converts a string into a int64.
 	 *
