MoviePipeline: Fixed an issue that could cause a crash in the graph when multiple output types are generated across multiple branches (regression from the recent filename disambiguation change).

diff --git a/Engine/Plugins/MovieScene/MovieRenderPipeline/Source/MovieRenderPipelineCore/Private/Graph/Nodes/MovieGraphVideoOutputNode.cpp b/Engine/Plugins/MovieScene/MovieRenderPipeline/Source/MovieRenderPipelineCore/Private/Graph/Nodes/MovieGraphVideoOutputNode.cpp
--- a/Engine/Plugins/MovieScene/MovieRenderPipeline/Source/MovieRenderPipelineCore/Private/Graph/Nodes/MovieGraphVideoOutputNode.cpp
+++ b/Engine/Plugins/MovieScene/MovieRenderPipeline/Source/MovieRenderPipelineCore/Private/Graph/Nodes/MovieGraphVideoOutputNode.cpp
@@ -77,12 +77,6 @@ void UMovieGraphVideoOutputNode::OnReceiveImageDataImpl(UMovieGraphPipeline* InP
 		FImagePixelData* RawRenderPassData = RenderPassData.Value.Get();
 		UE::MovieGraph::FMovieGraphSampleState* Payload = RawRenderPassData->GetPayload<UE::MovieGraph::FMovieGraphSampleState>();
 
-		constexpr bool bIncludeCDOs = false;
-		constexpr bool bExactMatch = true;
-		const UMovieGraphVideoOutputNode* ParentNode = Cast<UMovieGraphVideoOutputNode>(
-			InRawFrameData->EvaluatedConfig->GetSettingForBranch(GetClass(), RenderPassData.Key.RootBranchName, bIncludeCDOs, bExactMatch));
-		checkf(ParentNode, TEXT("Video output should not exist without a parent node in the graph."));
-
 		// This payload may have opted out of being written by this output type.
 		if (!Payload->OutputTypeRestrictions.IsEmpty() && !Payload->OutputTypeRestrictions.Contains(FSoftClassPath(GetClass())))
 		{
@@ -104,9 +98,16 @@ void UMovieGraphVideoOutputNode::OnReceiveImageDataImpl(UMovieGraphPipeline* InP
 			continue;
 		}
 
+		constexpr bool bIncludeCDOs = false;
+		constexpr bool bExactMatch = true;
+		const UMovieGraphVideoOutputNode* ParentNode = Cast<UMovieGraphVideoOutputNode>(
+			InRawFrameData->EvaluatedConfig->GetSettingForBranch(GetClass(), RenderPassData.Key.RootBranchName, bIncludeCDOs, bExactMatch));
+		checkf(ParentNode, TEXT("Video output should not exist without a parent node in the graph."));
+
 		PayloadsToWrite.FindOrAdd(ParentNode).Add(Payload);
 	}
 	
+	// NodeInstanceToPayloads needs to be populated before the loop below; it's used in filename disambiguation
 	InRawFrameData->NodeInstanceToPayloads.Append(PayloadsToWrite);
 
 	for (FMovieGraphPassData& RenderPassData : InRawFrameData->ImageOutputData)
@@ -115,7 +116,10 @@ void UMovieGraphVideoOutputNode::OnReceiveImageDataImpl(UMovieGraphPipeline* InP
 		constexpr bool bExactMatch = true;
 		const UMovieGraphVideoOutputNode* ParentNode = Cast<UMovieGraphVideoOutputNode>(
 			InRawFrameData->EvaluatedConfig->GetSettingForBranch(GetClass(), RenderPassData.Key.RootBranchName, bIncludeCDOs, bExactMatch));
-		checkf(ParentNode, TEXT("Video output should not exist without a parent node in the graph."));
+		if (!IsValid(ParentNode))
+		{
+			continue;
+		}
 
 		TArray<UE::MovieGraph::FMovieGraphSampleState*>* EligiblePayloads = PayloadsToWrite.Find(ParentNode);
 		if (!EligiblePayloads)
diff --git a/Engine/Plugins/MovieScene/MovieRenderPipeline/Source/MovieRenderPipelineRenderPasses/Private/MovieGraphImageSequenceOutputNode.cpp b/Engine/Plugins/MovieScene/MovieRenderPipeline/Source/MovieRenderPipelineRenderPasses/Private/MovieGraphImageSequenceOutputNode.cpp
--- a/Engine/Plugins/MovieScene/MovieRenderPipeline/Source/MovieRenderPipelineRenderPasses/Private/MovieGraphImageSequenceOutputNode.cpp
+++ b/Engine/Plugins/MovieScene/MovieRenderPipeline/Source/MovieRenderPipelineRenderPasses/Private/MovieGraphImageSequenceOutputNode.cpp
@@ -238,12 +238,6 @@ void UMovieGraphImageSequenceOutputNode::OnReceiveImageDataImpl(UMovieGraphPipel
 	{
 		checkf(RenderData.Value.IsValid(), TEXT("Unexpected empty image data: incorrectly moved or its production failed?"));
 		UE::MovieGraph::FMovieGraphSampleState* Payload = RenderData.Value->GetPayload<UE::MovieGraph::FMovieGraphSampleState>();
-		
-		constexpr bool bIncludeCDOs = false;
-		constexpr bool bExactMatch = true;
-		const UMovieGraphImageSequenceOutputNode* ParentNode = Cast<UMovieGraphImageSequenceOutputNode>(
-			InRawFrameData->EvaluatedConfig->GetSettingForBranch(GetClass(), RenderData.Key.RootBranchName, bIncludeCDOs, bExactMatch));
-		checkf(ParentNode, TEXT("Image sequence output should not exist without a parent node in the graph."));
 
 		// This payload may have opted out of being written by this output type.
 		if (!Payload->OutputTypeRestrictions.IsEmpty() && !Payload->OutputTypeRestrictions.Contains(FSoftClassPath(GetClass())))
@@ -265,10 +259,17 @@ void UMovieGraphImageSequenceOutputNode::OnReceiveImageDataImpl(UMovieGraphPipel
 		{
 			continue;
 		}
+		
+		constexpr bool bIncludeCDOs = false;
+        constexpr bool bExactMatch = true;
+        const UMovieGraphImageSequenceOutputNode* ParentNode = Cast<UMovieGraphImageSequenceOutputNode>(
+        	InRawFrameData->EvaluatedConfig->GetSettingForBranch(GetClass(), RenderData.Key.RootBranchName, bIncludeCDOs, bExactMatch));
+        checkf(ParentNode, TEXT("Image sequence output should not exist without a parent node in the graph."));
 
 		PayloadsToWrite.FindOrAdd(ParentNode).Add(Payload);
 	}
 	
+	// NodeInstanceToPayloads needs to be populated before the loop below; it's used in filename disambiguation
 	InRawFrameData->NodeInstanceToPayloads.Append(PayloadsToWrite);
 
 	// The base ImageSequenceOutputNode doesn't support any multilayer formats, so we write out each render pass separately.
@@ -278,7 +279,10 @@ void UMovieGraphImageSequenceOutputNode::OnReceiveImageDataImpl(UMovieGraphPipel
 		constexpr bool bExactMatch = true;
 		const UMovieGraphImageSequenceOutputNode* ParentNode = Cast<UMovieGraphImageSequenceOutputNode>(
 			InRawFrameData->EvaluatedConfig->GetSettingForBranch(GetClass(), RenderData.Key.RootBranchName, bIncludeCDOs, bExactMatch));
-		checkf(ParentNode, TEXT("Image sequence output should not exist without a parent node in the graph."));
+		if (!IsValid(ParentNode))
+		{
+			continue;
+		}
 
 		TArray<UE::MovieGraph::FMovieGraphSampleState*>* EligiblePayloads = PayloadsToWrite.Find(ParentNode);
 		if (!EligiblePayloads)
@@ -523,12 +527,6 @@ void UMovieGraphImageSequenceOutputNode_EXR::OnReceiveImageDataImpl(UMovieGraphP
 		checkf(RenderData.Value.IsValid(), TEXT("Unexpected empty image data: incorrectly moved or its production failed?"));
 		UE::MovieGraph::FMovieGraphSampleState* Payload = RenderData.Value->GetPayload<UE::MovieGraph::FMovieGraphSampleState>();
 
-		constexpr bool bIncludeCDOs = false;
-		constexpr bool bExactMatch = true;
-		const UMovieGraphImageSequenceOutputNode* ParentNode = Cast<UMovieGraphImageSequenceOutputNode>(
-			InRawFrameData->EvaluatedConfig->GetSettingForBranch(GetClass(), RenderData.Key.RootBranchName, bIncludeCDOs, bExactMatch));
-		checkf(ParentNode, TEXT("Image sequence output should not exist without a parent node in the graph."));
-
 		// This payload may have opted out of being written by this output type.
 		if (!Payload->OutputTypeRestrictions.IsEmpty() && !Payload->OutputTypeRestrictions.Contains(FSoftClassPath(GetClass())))
 		{
@@ -549,10 +547,17 @@ void UMovieGraphImageSequenceOutputNode_EXR::OnReceiveImageDataImpl(UMovieGraphP
 		{
 			continue;
 		}
+		
+		constexpr bool bIncludeCDOs = false;
+        constexpr bool bExactMatch = true;
+        const UMovieGraphImageSequenceOutputNode* ParentNode = Cast<UMovieGraphImageSequenceOutputNode>(
+        	InRawFrameData->EvaluatedConfig->GetSettingForBranch(GetClass(), RenderData.Key.RootBranchName, bIncludeCDOs, bExactMatch));
+        checkf(ParentNode, TEXT("Image sequence output should not exist without a parent node in the graph."));
 
 		PayloadsToWrite.FindOrAdd(ParentNode).Add(Payload);
 	}
 
+	// NodeInstanceToPayloads needs to be populated before the loop below; it's used in filename disambiguation
 	InRawFrameData->NodeInstanceToPayloads.Append(PayloadsToWrite);
 
 	for (TPair<FMovieGraphRenderDataIdentifier, TUniquePtr<FImagePixelData>>& RenderData : InRawFrameData->ImageOutputData)
@@ -561,7 +566,10 @@ void UMovieGraphImageSequenceOutputNode_EXR::OnReceiveImageDataImpl(UMovieGraphP
 		constexpr bool bExactMatch = true;
 		const UMovieGraphImageSequenceOutputNode_EXR* ParentNode = InRawFrameData->EvaluatedConfig->GetSettingForBranch<UMovieGraphImageSequenceOutputNode_EXR>(
 			RenderData.Key.RootBranchName, bIncludeCDOs, bExactMatch);
-		checkf(ParentNode, TEXT("Single-layer EXR should not exist without a parent node in the graph."));
+		if (!IsValid(ParentNode))
+		{
+			continue;
+		}
 
 		TArray<UE::MovieGraph::FMovieGraphSampleState*>* EligiblePayloads = PayloadsToWrite.Find(ParentNode);
 		if (!EligiblePayloads)
