MoviePipeline: Disable input handling in MRQ render window.

Resolves a SlateInvalidation issue with sharing HittestGrids between multiple slate paint passes (movie preview and offscreen rendering).

When resizing SRetainerWidget's HittestGrid while still appended to another windows HittestGrid, causes size mismatch issues. Keeping track of associated outer hittestgrids is not reliable, as the lifecycle of the outer hittestgrid cannot be guaranteed.

#rnx

diff --git a/Engine/Plugins/MovieScene/MovieRenderPipeline/Source/MovieRenderPipelineCore/Private/MoviePipelineInProcessExecutor.cpp b/Engine/Plugins/MovieScene/MovieRenderPipeline/Source/MovieRenderPipelineCore/Private/MoviePipelineInProcessExecutor.cpp
--- a/Engine/Plugins/MovieScene/MovieRenderPipeline/Source/MovieRenderPipelineCore/Private/MoviePipelineInProcessExecutor.cpp
+++ b/Engine/Plugins/MovieScene/MovieRenderPipeline/Source/MovieRenderPipelineCore/Private/MoviePipelineInProcessExecutor.cpp
@@ -191,6 +191,14 @@ void UMoviePipelineInProcessExecutor::OnTick()
 
 	FText WindowTitle = GetWindowTitle();
 	UKismetSystemLibrary::SetWindowTitle(WindowTitle);
+
+	if (UGameEngine* GameEngine = Cast<UGameEngine>(GEngine))
+	{
+		if (TSharedPtr<SWindow> GameViewportWindow = GameEngine->GameViewportWindow.Pin())
+		{
+			GameViewportWindow->SetAcceptsInput(false);
+		}
+	}
 }
 
 void UMoviePipelineInProcessExecutor::OnApplicationQuit()
@@ -308,6 +316,14 @@ void UMoviePipelineInProcessExecutor::RestoreState()
 			SavedState.WindowTitle.Reset();
 		}
 	}
+	
+	if (UGameEngine* GameEngine = Cast<UGameEngine>(GEngine))
+	{
+		if (TSharedPtr<SWindow> GameViewportWindow = GameEngine->GameViewportWindow.Pin())
+		{
+			GameViewportWindow->SetAcceptsInput(true);
+		}
+	}
 }
 void UMoviePipelineInProcessExecutor::CancelAllJobs_Implementation()
 {
diff --git a/Engine/Source/Runtime/SlateCore/Private/Widgets/SWindow.cpp b/Engine/Source/Runtime/SlateCore/Private/Widgets/SWindow.cpp
--- a/Engine/Source/Runtime/SlateCore/Private/Widgets/SWindow.cpp
+++ b/Engine/Source/Runtime/SlateCore/Private/Widgets/SWindow.cpp
@@ -1569,7 +1569,8 @@ EWindowActivationPolicy SWindow::ActivationPolicy() const
 /** @return true if the window accepts input; false if the window is non-interactive */
 bool SWindow::AcceptsInput() const
 {
-	return Type != EWindowType::CursorDecorator && (Type != EWindowType::ToolTip || !FSlateApplicationBase::Get().IsWindowHousingInteractiveTooltip(SharedThis(this)));
+	return bAcceptsInput && Type != EWindowType::CursorDecorator
+		&& (Type != EWindowType::ToolTip || !FSlateApplicationBase::Get().IsWindowHousingInteractiveTooltip(SharedThis(this)));
 }
 
 /** @return true if the user decides the size of the window; false if the content determines the size of the window */
@@ -2046,6 +2047,7 @@ SWindow::SWindow()
 	, bShouldPreserveAspectRatio( false )
 	, bManualManageDPI( false )
 	, bAllowFastUpdate( false )
+	, bAcceptsInput( true )
 	, WindowActivationPolicy( EWindowActivationPolicy::Always )
 	, InitialDesiredScreenPosition( FVector2f::ZeroVector )
 	, InitialDesiredSize( FVector2f::ZeroVector )
diff --git a/Engine/Source/Runtime/SlateCore/Public/Widgets/SWindow.h b/Engine/Source/Runtime/SlateCore/Public/Widgets/SWindow.h
--- a/Engine/Source/Runtime/SlateCore/Public/Widgets/SWindow.h
+++ b/Engine/Source/Runtime/SlateCore/Public/Widgets/SWindow.h
@@ -803,6 +803,9 @@ public:
 	/** @return true if the window accepts input; false if the window is non-interactive */
 	SLATECORE_API bool AcceptsInput() const;
 
+	/** If true the window will not accept input regardless of its type*/
+	SLATECORE_API void SetAcceptsInput(bool bEnabled) { bAcceptsInput = bEnabled; }
+
 	/** @return true if the user decides the size of the window */
 	SLATECORE_API bool IsUserSized() const;
 
@@ -1196,6 +1199,9 @@ protected:
 	/** True if this window allows global invalidation of its contents */
 	bool bAllowFastUpdate : 1;
 
+	/** If true, window will accept user inputs */
+	bool bAcceptsInput: 1;
+	
 	/** When should the window be activated upon being shown */
 	EWindowActivationPolicy WindowActivationPolicy;
 
