MoviePipeline: Fixed an issue that caused Quick Renders which use subgraphs to fail.

Subgraphs would sometimes be garbage collected, which resulted in incomplete graph evaluation.

diff --git a/Engine/Plugins/MovieScene/MovieRenderPipeline/Source/MovieRenderPipelineEditor/Private/Graph/MovieGraphQuickRender.cpp b/Engine/Plugins/MovieScene/MovieRenderPipeline/Source/MovieRenderPipelineEditor/Private/Graph/MovieGraphQuickRender.cpp
--- a/Engine/Plugins/MovieScene/MovieRenderPipeline/Source/MovieRenderPipelineEditor/Private/Graph/MovieGraphQuickRender.cpp
+++ b/Engine/Plugins/MovieScene/MovieRenderPipeline/Source/MovieRenderPipelineEditor/Private/Graph/MovieGraphQuickRender.cpp
@@ -351,7 +351,6 @@ bool UMovieGraphQuickRenderSubsystem::GenerateAndPopulateQueue(TObjectPtr<UMovie
 	OutJob->JobName = FString(TEXT("QuickRender"));
 
 	// Update the job to use the quick render graph (either the default, or the one specified in the settings)
-	TMap<TObjectPtr<UMovieGraphConfig>, TObjectPtr<UMovieGraphConfig>> OriginalGraphToDupeMap;
 	const UMovieGraphConfig* QuickRenderGraph = GetQuickRenderGraph(InQuickRenderSettings->GraphPreset.LoadSynchronous(), OriginalGraphToDupeMap);
 	if (!QuickRenderGraph)
 	{
@@ -395,6 +394,7 @@ void UMovieGraphQuickRenderSubsystem::PerformPostRenderCleanup()
 	TemporaryGraph = nullptr;
 	TemporaryExecutor = nullptr;
 	TemporaryEvaluatedGraph = nullptr;
+	OriginalGraphToDupeMap.Empty();
 	RenderingLevelSequence = nullptr;
 	UtilityLevelSequence = nullptr;
 	QuickRenderModeSettings = nullptr;
diff --git a/Engine/Plugins/MovieScene/MovieRenderPipeline/Source/MovieRenderPipelineEditor/Public/Graph/MovieGraphQuickRender.h b/Engine/Plugins/MovieScene/MovieRenderPipeline/Source/MovieRenderPipelineEditor/Public/Graph/MovieGraphQuickRender.h
--- a/Engine/Plugins/MovieScene/MovieRenderPipeline/Source/MovieRenderPipelineEditor/Public/Graph/MovieGraphQuickRender.h
+++ b/Engine/Plugins/MovieScene/MovieRenderPipeline/Source/MovieRenderPipelineEditor/Public/Graph/MovieGraphQuickRender.h
@@ -184,6 +184,10 @@ private:
 	UPROPERTY(Transient)
 	TObjectPtr<UMovieGraphEvaluatedConfig> TemporaryEvaluatedGraph = nullptr;
 
+	/** Maps a graph to its duplicate (which is done during render setup). Includes subgraphs. */
+	UPROPERTY(Transient)
+	TMap<TObjectPtr<UMovieGraphConfig>, TObjectPtr<UMovieGraphConfig>> OriginalGraphToDupeMap;
+
 	/**
 	 * The level sequence that is used for rendering (ie, provided directly to MRG to render from). Sometimes this will be the current level sequence
 	 * in Sequencer, and other times it may be a temporary level sequence solely for use in Quick Render. It depends on the mode that is being used.
