MoviePipeline: Fixed an issue where a long job comment in the burn-in would not perform a layout properly on the first frame (fixed in both the graph and presets).

Also updated the default burn-in UMG assets to use TextBlock instead of TextBox widgets to prevent scrollbars from showing up in some scenarios.

#virtualized

diff --git a/Engine/Plugins/MovieScene/MovieRenderPipeline/Source/MovieRenderPipelineCore/Private/Graph/Nodes/MovieGraphWidgetRendererBaseNode.cpp b/Engine/Plugins/MovieScene/MovieRenderPipeline/Source/MovieRenderPipelineCore/Private/Graph/Nodes/MovieGraphWidgetRendererBaseNode.cpp
--- a/Engine/Plugins/MovieScene/MovieRenderPipeline/Source/MovieRenderPipelineCore/Private/Graph/Nodes/MovieGraphWidgetRendererBaseNode.cpp
+++ b/Engine/Plugins/MovieScene/MovieRenderPipeline/Source/MovieRenderPipelineCore/Private/Graph/Nodes/MovieGraphWidgetRendererBaseNode.cpp
@@ -176,9 +176,19 @@ void UMovieGraphWidgetRendererBaseNode::FMovieGraphWidgetPass::Render(const FMov
 		VirtualWindow->SetContent(WidgetToRender.ToSharedRef());
 
 		// Draw the widget to the render target.  This leaves the texture in SRV state so no transition is needed.
+		// Note: If this draw call is updated, be sure to also update the one below.
 		NodeCDO->WidgetRenderer->DrawWindow(RenderTarget, VirtualWindow->GetHittestGrid(), VirtualWindow.ToSharedRef(),
 			1.f, OutputResolution, InTimeData.FrameDeltaTime);
 
+		// If this is the first frame being rendered, then re-render. This will ensure that all layout computations are correct (for example, large
+		// text blocks may not have their layouts computed by the time the first render above occurs, especially if the text is driven by the
+		// UMG blueprint).
+		if (InFrameTraversalContext.Time.ShotOutputFrameNumber == 0)
+		{
+			NodeCDO->WidgetRenderer->DrawWindow(RenderTarget, VirtualWindow->GetHittestGrid(), VirtualWindow.ToSharedRef(),
+				1.f, OutputResolution, InTimeData.FrameDeltaTime);
+		}
+
 		ENQUEUE_RENDER_COMMAND(WidgetRenderTargetResolveCommand)(
 			[
 				InTimeData,
diff --git a/Engine/Plugins/MovieScene/MovieRenderPipeline/Source/MovieRenderPipelineSettings/Private/MoviePipelineBurnInSetting.cpp b/Engine/Plugins/MovieScene/MovieRenderPipeline/Source/MovieRenderPipelineSettings/Private/MoviePipelineBurnInSetting.cpp
--- a/Engine/Plugins/MovieScene/MovieRenderPipeline/Source/MovieRenderPipelineSettings/Private/MoviePipelineBurnInSetting.cpp
+++ b/Engine/Plugins/MovieScene/MovieRenderPipeline/Source/MovieRenderPipelineSettings/Private/MoviePipelineBurnInSetting.cpp
@@ -82,7 +82,18 @@ void UMoviePipelineBurnInSetting::RenderSample_GameThreadImpl(const FMoviePipeli
 			CurrentWidget->OnOutputFrameStarted(GetPipeline());
 
 			// Draw the widget to the render target. This leaves the texture in SRV state so no transition is needed.
-			WidgetRenderer->DrawWindow(RenderTarget, VirtualWindow->GetHittestGrid(), VirtualWindow.ToSharedRef(), 1.f, OutputResolution, InSampleState.OutputState.TimeData.FrameDeltaTime);
+			// Note: If this draw call is updated, be sure to also update the one below.
+			WidgetRenderer->DrawWindow(RenderTarget, VirtualWindow->GetHittestGrid(), VirtualWindow.ToSharedRef(),
+				1.f, OutputResolution, InSampleState.OutputState.TimeData.FrameDeltaTime);
+
+			// If this is the first frame being rendered, then re-render. This will ensure that all layout computations are correct (for example, large
+			// text blocks may not have their layouts computed by the time the first render above occurs, especially if the text is driven by the
+			// UMG blueprint).
+			if (InSampleState.OutputState.ShotOutputFrameNumber == 0)
+			{
+				WidgetRenderer->DrawWindow(RenderTarget, VirtualWindow->GetHittestGrid(), VirtualWindow.ToSharedRef(),
+					1.f, OutputResolution, InSampleState.OutputState.TimeData.FrameDeltaTime);
+			}
 
 			FRenderTarget* BackbufferRenderTarget = RenderTarget->GameThread_GetRenderTargetResource();
 			TSharedPtr<FMoviePipelineOutputMerger, ESPMode::ThreadSafe> OutputBuilder = GetPipeline()->OutputBuilder;
