MoviePipeline: Added support for MRG to MoviePipelineFunctionalTestBase.

diff --git a/Engine/Plugins/MovieScene/MovieRenderPipeline/Source/MovieRenderPipelineCore/Public/MovieRenderPipelineDataTypes.h b/Engine/Plugins/MovieScene/MovieRenderPipeline/Source/MovieRenderPipelineCore/Public/MovieRenderPipelineDataTypes.h
--- a/Engine/Plugins/MovieScene/MovieRenderPipeline/Source/MovieRenderPipelineCore/Public/MovieRenderPipelineDataTypes.h
+++ b/Engine/Plugins/MovieScene/MovieRenderPipeline/Source/MovieRenderPipelineCore/Public/MovieRenderPipelineDataTypes.h
@@ -129,6 +129,11 @@ struct FMoviePipelinePassIdentifier
 		return HashCombine(GetTypeHash(OutputState.Name), GetTypeHash(OutputState.CameraName));
 	}
 
+	friend FString LexToString(const FMoviePipelinePassIdentifier InIdentifier)
+	{
+		return FString::Printf(TEXT("Name: %s, Camera: %s"), *InIdentifier.Name, *InIdentifier.CameraName);
+	}
+
 public:
 	// The name of the pass such as "FinalImage" or "ObjectId", etc.
 	UPROPERTY(EditAnywhere, BlueprintReadWrite, Category = "Movie Pipeline")
diff --git a/Engine/Plugins/MovieScene/MovieRenderPipeline/Source/MovieRenderPipelineEditor/Private/MoviePipelineFunctionalTestBase.cpp b/Engine/Plugins/MovieScene/MovieRenderPipeline/Source/MovieRenderPipelineEditor/Private/MoviePipelineFunctionalTestBase.cpp
--- a/Engine/Plugins/MovieScene/MovieRenderPipeline/Source/MovieRenderPipelineEditor/Private/MoviePipelineFunctionalTestBase.cpp
+++ b/Engine/Plugins/MovieScene/MovieRenderPipeline/Source/MovieRenderPipelineEditor/Private/MoviePipelineFunctionalTestBase.cpp
@@ -19,6 +19,8 @@
 #include "Interfaces/IScreenShotToolsModule.h"
 #include "AutomationBlueprintFunctionLibrary.h"
 #include "AutomationWorkerMessages.h"
+#include "Graph/MovieGraphBlueprintLibrary.h"
+#include "Graph/MovieGraphPipeline.h"
 
 #include UE_INLINE_GENERATED_CPP_BY_NAME(MoviePipelineFunctionalTestBase)
 
@@ -54,21 +56,38 @@ bool AMoviePipelineFunctionalTestBase::IsReady_Implementation()
 
 void AMoviePipelineFunctionalTestBase::StartTest()
 {
+	ActiveQueue = NewObject<UMoviePipelineQueue>(GetWorld());
+	ActiveQueue->CopyFrom(QueuePreset);
+
+	// PrepareTest() guarantees exactly one job in the queue.
+	UMoviePipelineExecutorJob* Job = ActiveQueue->GetJobs()[0];
 
 	// PIE will already be running at this point so we want to instantiate an instance of
 	// the Movie Pipeline in the current world and just run it. This doesn't test the UI/PIE 
 	// portion of the system but that is more stable than the actual featureset.
-	ActiveMoviePipeline = NewObject<UMoviePipeline>(GetWorld());
+	if (Job->IsUsingGraphConfiguration())
+	{
+		ActiveMoviePipeline = NewObject<UMovieGraphPipeline>(GetWorld());
+	}
+	else
+	{
+		ActiveMoviePipeline = NewObject<UMoviePipeline>(GetWorld());
+	}
+
 	ActiveMoviePipeline->OnMoviePipelineWorkFinished().AddUObject(this, &AMoviePipelineFunctionalTestBase::OnMoviePipelineFinished);
 	ActiveMoviePipeline->OnMoviePipelineShotWorkFinished().AddUObject(this, &AMoviePipelineFunctionalTestBase::OnJobShotFinished);
 
-	ActiveQueue = NewObject<UMoviePipelineQueue>(GetWorld());
-	ActiveQueue->CopyFrom(QueuePreset);
-
 	// Ensure we've initialized any transient settings (ie: the game overrides setting that is automatically added),
 	// otherwise it won't get called.
-	ActiveQueue->GetJobs()[0]->GetConfiguration()->InitializeTransientSettings();
-	ActiveMoviePipeline->Initialize(ActiveQueue->GetJobs()[0]); // Prepare Test will ensure we have one job
+	if (UMoviePipeline* PipelineAsLegacy = Cast<UMoviePipeline>(ActiveMoviePipeline))
+	{
+		Job->GetConfiguration()->InitializeTransientSettings();
+		PipelineAsLegacy->Initialize(Job);
+	}
+	else if (UMovieGraphPipeline* PipelineAsGraph = Cast<UMovieGraphPipeline>(ActiveMoviePipeline))
+	{
+		PipelineAsGraph->Initialize(Job, FMovieGraphInitConfig());
+	}
 }
 
 void AMoviePipelineFunctionalTestBase::OnJobShotFinished(FMoviePipelineOutputData InOutputData)
@@ -104,41 +123,64 @@ bool SaveOutputToGroundTruth(const FString& GroundTruthDirectory, FMoviePipeline
 	// of the files from the old location to the new location. We want to keep these relative to
 	// the original output directory from MRQ, so that we can make tests that ensure sub-folder
 	// structures get generated correctly.
-	UMoviePipelineOutputSetting* OutputSetting = Cast<UMoviePipelineOutputSetting>(OutputData.Job->GetConfiguration()->FindOrAddSettingByClass(UMoviePipelineOutputSetting::StaticClass()));
 	FString OriginalRootOutputDirectory = UMoviePipelineEditorBlueprintLibrary::ResolveOutputDirectoryFromJob(OutputData.Job);
 
-	for (FMoviePipelineShotOutputData& Shot : OutputData.ShotData)
+	// Gather the generated file paths for presets and graphs
+	TArray<FString*> GeneratedFilePaths;
+	if (OutputData.Job->IsUsingGraphConfiguration())
 	{
-		for (TPair<FMoviePipelinePassIdentifier, FMoviePipelineRenderPassOutputData>& Pair : Shot.RenderPassData)
+		for (FMovieGraphRenderOutputData& Shot : OutputData.GraphData)
 		{
-			for (FString& FilePath : Pair.Value.FilePaths)
+			for (TPair<FMovieGraphRenderDataIdentifier, FMovieGraphRenderLayerOutputData>& RenderLayerData : Shot.RenderLayerData)
 			{
-				FString RelativePath = FilePath;
-				if (!FPaths::MakePathRelativeTo(/*Out*/ RelativePath, *OriginalRootOutputDirectory))
+				for (FString& FilePath : RenderLayerData.Value.FilePaths)
 				{
-					UE_LOG(LogTemp, Error, TEXT("Could not generate ground truth, unable to generate relative paths."));
-					return false;
+					GeneratedFilePaths.Add(&FilePath);
 				}
-
-				FString NewPath = GroundTruthDirectory / RelativePath;
-				FString AbsoluteNewPath = FPaths::ConvertRelativePathToFull(NewPath);
-				if (IFileManager::Get().Copy(*AbsoluteNewPath, *FilePath) != ECopyResult::COPY_OK)
+			}
+		}
+	}
+	else
+	{
+		for (FMoviePipelineShotOutputData& Shot : OutputData.ShotData)
+		{
+			for (TPair<FMoviePipelinePassIdentifier, FMoviePipelineRenderPassOutputData>& Pair : Shot.RenderPassData)
+			{
+				for (FString& FilePath : Pair.Value.FilePaths)
 				{
-					UE_LOG(LogTemp, Error, TEXT("Could not generate ground truth, unable to copy existing image to Automation directory."));
-					return false;
+					GeneratedFilePaths.Add(&FilePath);
 				}
+			}
+		}
+	}
 
-				// Now rewrite the FilePath in the struct to the new location so that when we save the json below it has the right path.
-				FString RelativeToAutomationDir = NewPath;
-				if (!FPaths::MakePathRelativeTo(/*InOut*/ RelativeToAutomationDir, *GroundTruthDirectory))
-				{
-					UE_LOG(LogTemp, Error, TEXT("Could not generate ground truth, unable to generate relative paths."));
-					return false;
-				}
+	// With all generated file paths known, rewrite them to be relative to the automation directory
+	for (FString* FilePath : GeneratedFilePaths)
+	{
+		FString RelativePath = *FilePath;
+		if (!FPaths::MakePathRelativeTo(/*Out*/ RelativePath, *OriginalRootOutputDirectory))
+		{
+			UE_LOG(LogTemp, Error, TEXT("Could not generate ground truth, unable to generate relative paths."));
+			return false;
+		}
 
-				FilePath = RelativeToAutomationDir;
-			}
+		FString NewPath = GroundTruthDirectory / RelativePath;
+		FString AbsoluteNewPath = FPaths::ConvertRelativePathToFull(NewPath);
+		if (IFileManager::Get().Copy(*AbsoluteNewPath, GetData(*FilePath)) != ECopyResult::COPY_OK)
+		{
+			UE_LOG(LogTemp, Error, TEXT("Could not generate ground truth, unable to copy existing image to Automation directory."));
+			return false;
 		}
+
+		// Now rewrite the FilePath in the struct to the new location so that when we save the json below it has the right path.
+		FString RelativeToAutomationDir = NewPath;
+		if (!FPaths::MakePathRelativeTo(/*InOut*/ RelativeToAutomationDir, *GroundTruthDirectory))
+		{
+			UE_LOG(LogTemp, Error, TEXT("Could not generate ground truth, unable to generate relative paths."));
+			return false;
+		}
+
+		*FilePath = RelativeToAutomationDir;
 	}
 
 	// Now that we've copied all of the images across we can serialize the struct to a json string
@@ -153,12 +195,19 @@ bool SaveOutputToGroundTruth(const FString& GroundTruthDirectory, FMoviePipeline
 
 void AMoviePipelineFunctionalTestBase::CompareRenderOutputToGroundTruth(FMoviePipelineOutputData InOutputData)
 {
-	// Grab the (expected) resolution from out rendered data
-	UMoviePipelineOutputSetting* OutputSetting = Cast< UMoviePipelineOutputSetting>(InOutputData.Job->GetConfiguration()->FindOrAddSettingByClass(UMoviePipelineOutputSetting::StaticClass()));
-	check(OutputSetting);
+	const bool bIsUsingGraph = InOutputData.Job->IsUsingGraphConfiguration();
+	
+	// Grab the (expected) resolution from the configuration (graph or preset)
+	FString ErrorMsg;
+	FIntPoint OutputResolution = GetExpectedOutputResolution(InOutputData, ErrorMsg);
+	if (!ErrorMsg.IsEmpty())
+	{
+		FinishTest(EFunctionalTestResult::Failed, *ErrorMsg);
+		return;
+	}
 
 	// Build screenshot data for this test. This contains a lot of metadata about RHI, Platform, etc that we'll use to generate the output folder name.
-	FAutomationScreenshotData Data = UAutomationBlueprintFunctionLibrary::BuildScreenshotData(GetWorld(), TestLabel, OutputSetting->OutputResolution.X, OutputSetting->OutputResolution.Y);
+	FAutomationScreenshotData Data = UAutomationBlueprintFunctionLibrary::BuildScreenshotData(GetWorld(), TestLabel, OutputResolution.X, OutputResolution.Y);
 
 	// Convert the Screenshot Data into Metadata
 	const FImageTolerance ImageTolerance = GetImageToleranceForPreset(ImageToleranceLevel, CustomImageTolerance);
@@ -202,8 +251,8 @@ void AMoviePipelineFunctionalTestBase::CompareRenderOutputToGroundTruth(FMoviePi
 
 	// Do some basic checks on our new data to ensure it output the expected number of files/shots/etc, before doing
 	// the computationally expensive image comparisons.
-	int32 NumShotsNew = InOutputData.ShotData.Num();
-	int32 NumShotsGT = GroundTruthData.ShotData.Num();
+	int32 NumShotsNew = bIsUsingGraph ? InOutputData.GraphData.Num() : InOutputData.ShotData.Num();
+	int32 NumShotsGT = bIsUsingGraph ? GroundTruthData.GraphData.Num() : GroundTruthData.ShotData.Num();
 	if (NumShotsNew != NumShotsGT)
 	{
 		FinishTest(EFunctionalTestResult::Failed, FString::Printf(TEXT("Mismatched number of shots between GroudTruth and New. Expected %d got %d."), NumShotsGT, NumShotsNew));
@@ -212,48 +261,25 @@ void AMoviePipelineFunctionalTestBase::CompareRenderOutputToGroundTruth(FMoviePi
 		
 	TMap<FString, FString> OldToNewImagesToCompare;
 	
-	for (int32 ShotIndex = 0; ShotIndex < GroundTruthData.ShotData.Num(); ShotIndex++)
+	for (int32 ShotIndex = 0; ShotIndex < NumShotsGT; ShotIndex++)
 	{
-		// We compare them using indexes because we know that they have the same number of indexes,
-		// and the order should be deterministic. 
-		const FMoviePipelineShotOutputData & GroundTruthShot = GroundTruthData.ShotData[ShotIndex];
-		const FMoviePipelineShotOutputData & NewData = InOutputData.ShotData[ShotIndex];
-		
 		// And then repeat this for each render pass in the shot
-		int32 NumRenderPassesNew = NewData.RenderPassData.Num();
-		int32 NumRenderPassesGT = GroundTruthShot.RenderPassData.Num();
+		int32 NumRenderPassesNew = bIsUsingGraph ? InOutputData.GraphData[ShotIndex].RenderLayerData.Num() : InOutputData.ShotData[ShotIndex].RenderPassData.Num();
+		int32 NumRenderPassesGT = bIsUsingGraph ? GroundTruthData.GraphData[ShotIndex].RenderLayerData.Num() : GroundTruthData.ShotData[ShotIndex].RenderPassData.Num();
 		if (NumRenderPassesNew != NumRenderPassesGT)
 		{
 			FinishTest(EFunctionalTestResult::Failed, FString::Printf(TEXT("Mismatched number of shots between GroudTruth and New. Expected %d got %d."), NumRenderPassesGT, NumRenderPassesNew));
 			return;
 		}
-		
-		for (const TPair<FMoviePipelinePassIdentifier, FMoviePipelineRenderPassOutputData>& GroundTruthPair : GroundTruthShot.RenderPassData)
+
+		// Update the map that links a ground truth image to its equivalent test-generated image.
+		if (bIsUsingGraph)
 		{
-			const FMoviePipelineRenderPassOutputData* NewDataOutputData = NewData.RenderPassData.Find(GroundTruthPair.Key);
-			if (!NewDataOutputData)
-			{
-				FinishTest(EFunctionalTestResult::Failed, FString::Printf(TEXT("Did not render pass '%s' from camera '%s' that is in the Ground Truth."), *GroundTruthPair.Key.Name, *GroundTruthPair.Key.CameraName));
-				return;
-			}
-			
-			// Now we can check that they output the same number of files
-			int32 NumOutputFilesNew = NewDataOutputData->FilePaths.Num();
-			int32 NumOutputFilesGT = GroundTruthPair.Value.FilePaths.Num();
-			if (NumOutputFilesNew != NumOutputFilesGT)
-			{
-				FinishTest(EFunctionalTestResult::Failed, FString::Printf(TEXT("Mismatched number of shots between GroudTruth and New. Expected %d got %d."), NumOutputFilesGT, NumOutputFilesNew));
-				return;
-			}
-			
-			// Now we'll just loop through them in lockstep (as MRQ should put them in order anyways) and just add
-			// them to an array to be tested later. We don't test that the sub-file names are exact matches right now.
-			for (int32 Index = 0; Index < NumOutputFilesGT; Index++)
-			{
-				FString AbsoluteGTPath = FPaths::ConvertRelativePathToFull(FPaths::CreateStandardFilename(ReportDirectory / GroundTruthPair.Value.FilePaths[Index]));
-				FString AbsoluteNewPath = FPaths::ConvertRelativePathToFull(FPaths::CreateStandardFilename(NewDataOutputData->FilePaths[Index]));
-				OldToNewImagesToCompare.Add(AbsoluteGTPath, AbsoluteNewPath);
-			}
+			UpdateGroundTruthToTestImageMap(GroundTruthData.GraphData[ShotIndex].RenderLayerData, InOutputData.GraphData[ShotIndex].RenderLayerData, ReportDirectory, OldToNewImagesToCompare);
+		}
+		else
+		{
+			UpdateGroundTruthToTestImageMap(GroundTruthData.ShotData[ShotIndex].RenderPassData, InOutputData.ShotData[ShotIndex].RenderPassData, ReportDirectory, OldToNewImagesToCompare);
 		}
 	}
 		
@@ -274,3 +300,76 @@ void AMoviePipelineFunctionalTestBase::CompareRenderOutputToGroundTruth(FMoviePi
 	FinishTest(EFunctionalTestResult::Succeeded, TEXT(""));
 }
 
+FIntPoint AMoviePipelineFunctionalTestBase::GetExpectedOutputResolution(const FMoviePipelineOutputData& InOutputData, FString& OutErrorMsg) const
+{
+	OutErrorMsg.Empty();
+	
+	const bool bIsUsingGraph = InOutputData.Job->IsUsingGraphConfiguration();
+	
+	// Grab the (expected) resolution from the configuration (graph or preset)
+	FIntPoint OutputResolution(0, 0);
+	if (bIsUsingGraph)
+	{
+		FMovieGraphTraversalContext TraversalContext;
+		TraversalContext.Job = InOutputData.Job;
+
+		FString OutEvaluationError;
+		UMovieGraphEvaluatedConfig* EvaluatedGraph = InOutputData.Job->GetGraphPreset()->CreateFlattenedGraph(TraversalContext, OutEvaluationError);
+		if (!EvaluatedGraph)
+		{
+			OutErrorMsg = FString::Printf(TEXT("Unable to create evaluated graph: %s"), *OutEvaluationError);
+			return OutputResolution;
+		}
+
+		OutputResolution = UMovieGraphBlueprintLibrary::GetDesiredOutputResolution(EvaluatedGraph);
+	}
+	else
+	{
+		const UMoviePipelineOutputSetting* OutputSetting = Cast<UMoviePipelineOutputSetting>(InOutputData.Job->GetConfiguration()->FindOrAddSettingByClass(UMoviePipelineOutputSetting::StaticClass()));
+		if (!OutputSetting)
+		{
+			OutErrorMsg = FString(TEXT("Unable to find the Output setting in the configuration."));
+			return OutputResolution;
+		}
+		
+		OutputResolution = OutputSetting->OutputResolution;
+	}
+
+	return OutputResolution;
+}
+
+template<typename IdentifierType, typename OutputDataType>
+void AMoviePipelineFunctionalTestBase::UpdateGroundTruthToTestImageMap(
+	const TMap<IdentifierType, OutputDataType>& InGroundTruthData,
+	const TMap<IdentifierType, OutputDataType>& InTestData,
+	const FString& InReportDirectory,
+	TMap<FString, FString>& OutGroundTruthToTestImageMap)
+{
+	for (const TPair<IdentifierType, OutputDataType>& GroundTruthPair : InGroundTruthData)
+	{
+		const OutputDataType* TestOutputData = InTestData.Find(GroundTruthPair.Key);
+		if (!TestOutputData)
+		{
+			FinishTest(EFunctionalTestResult::Failed, FString::Printf(TEXT("Did not find render pass [%s] that is in the Ground Truth."), *LexToString(GroundTruthPair.Key)));
+			return;
+		}
+		
+		// Now we can check that they output the same number of files
+		const int32 NumOutputFilesTest = TestOutputData->FilePaths.Num();
+		const int32 NumOutputFilesGT = GroundTruthPair.Value.FilePaths.Num();
+		if (NumOutputFilesTest != NumOutputFilesGT)
+		{
+			FinishTest(EFunctionalTestResult::Failed, FString::Printf(TEXT("Mismatched number of shots between GroudTruth and New. Expected %d, got %d."), NumOutputFilesGT, NumOutputFilesTest));
+			return;
+		}
+			
+		// Now we'll just loop through them in lockstep (as MRQ should put them in order anyways) and just add
+		// them to an array to be tested later. We don't test that the sub-file names are exact matches right now.
+		for (int32 Index = 0; Index < NumOutputFilesGT; Index++)
+		{
+			FString AbsoluteGTPath = FPaths::ConvertRelativePathToFull(FPaths::CreateStandardFilename(InReportDirectory / GroundTruthPair.Value.FilePaths[Index]));
+			FString AbsoluteNewPath = FPaths::ConvertRelativePathToFull(FPaths::CreateStandardFilename(TestOutputData->FilePaths[Index]));
+			OutGroundTruthToTestImageMap.Add(AbsoluteGTPath, AbsoluteNewPath);
+		}
+	}
+}
\ No newline at end of file
diff --git a/Engine/Plugins/MovieScene/MovieRenderPipeline/Source/MovieRenderPipelineEditor/Private/MoviePipelineFunctionalTestBase.h b/Engine/Plugins/MovieScene/MovieRenderPipeline/Source/MovieRenderPipelineEditor/Private/MoviePipelineFunctionalTestBase.h
--- a/Engine/Plugins/MovieScene/MovieRenderPipeline/Source/MovieRenderPipelineEditor/Private/MoviePipelineFunctionalTestBase.h
+++ b/Engine/Plugins/MovieScene/MovieRenderPipeline/Source/MovieRenderPipelineEditor/Private/MoviePipelineFunctionalTestBase.h
@@ -43,6 +43,26 @@ protected:
 	UE_API void OnMoviePipelineFinished(FMoviePipelineOutputData InOutputData);
 	UE_API void CompareRenderOutputToGroundTruth(FMoviePipelineOutputData InOutputData);
 
+private:
+	/**
+	 * Gets the resolution that the test-generated images are expected to be.
+	 *
+	 * If there was an error getting this, OutErrorMsg will be populated with the error message, and the returned point will be (0, 0).
+	 */
+	FIntPoint GetExpectedOutputResolution(const FMoviePipelineOutputData& InOutputData, FString& OutErrorMsg) const;
+
+	/**
+	 * Updates OutGroundTruthToTestImageMap by linking frames in InGroundTruthData to those in InTestData.
+	 *
+	 * Keys in OutGroundTruthToTestImageMap are the ground-truth frame path (absolute), and values are the test-generated frame path (absolute).
+	 */
+	template<typename IdentifierType, typename OutputDataType>
+	void UpdateGroundTruthToTestImageMap(
+		const TMap<IdentifierType, OutputDataType>& InGroundTruthData,
+		const TMap<IdentifierType, OutputDataType>& InTestData,
+		const FString& InReportDirectory,
+		TMap<FString, FString>& OutGroundTruthToTestImageMap);
+
 protected:
 
 	UPROPERTY(EditAnywhere, BlueprintReadOnly, Category = "Movie Pipeline")
@@ -64,7 +84,7 @@ protected:
 
 protected:
 	UPROPERTY(Transient)
-	TObjectPtr<UMoviePipeline> ActiveMoviePipeline;
+	TObjectPtr<UMoviePipelineBase> ActiveMoviePipeline;
 
 	UPROPERTY(Transient)
 	TObjectPtr<UMoviePipelineQueue> ActiveQueue;
