MoviePipeline: Fixed an issue where job variable values set in the queue may not be committed properly if another job is selected before the widget loses focus.

diff --git a/Engine/Plugins/MovieScene/MovieRenderPipeline/Source/MovieRenderPipelineEditor/Private/Widgets/SMoviePipelineQueuePanel.cpp b/Engine/Plugins/MovieScene/MovieRenderPipeline/Source/MovieRenderPipelineEditor/Private/Widgets/SMoviePipelineQueuePanel.cpp
--- a/Engine/Plugins/MovieScene/MovieRenderPipeline/Source/MovieRenderPipelineEditor/Private/Widgets/SMoviePipelineQueuePanel.cpp
+++ b/Engine/Plugins/MovieScene/MovieRenderPipeline/Source/MovieRenderPipelineEditor/Private/Widgets/SMoviePipelineQueuePanel.cpp
@@ -497,6 +497,13 @@ void SMoviePipelineQueuePanel::OnSelectionChanged(const TArray<UMoviePipelineExe
 	{
 		Jobs.Append(InSelectedJobs);
 	}
+
+	if (!JobDetailsPanelWidget->GetSelectedObjects().IsEmpty())
+	{
+		// Force a focus lost event before the selection changes. Depending on how a property was added to the details panel, it may not have time to
+		// properly commit its value if its widget has not lost focus before the selection is changed.
+		FSlateApplication::Get().ClearKeyboardFocus(EFocusCause::SetDirectly);
+	}
 	
 	JobDetailsPanelWidget->SetObjects(Jobs);
 	NumSelectedJobs = InSelectedJobs.Num();
