Added MinPriority to SetWithCurrentPriority and ReplaceCurrentPriorityAndTag, it represents the minimum priority we will set that cvar.
Improve the log reporting when SetWithCurrentPriority and ReplaceCurrentPriorityAndTag are used to replace SetByConstructor values
Fix usage in the MRQ pipeline that was triggering the log

diff --git a/Engine/Plugins/MovieScene/MovieRenderPipeline/Source/MovieRenderPipelineCore/Private/Graph/MovieGraphDefaultAudioRenderer.cpp b/Engine/Plugins/MovieScene/MovieRenderPipeline/Source/MovieRenderPipelineCore/Private/Graph/MovieGraphDefaultAudioRenderer.cpp
--- a/Engine/Plugins/MovieScene/MovieRenderPipeline/Source/MovieRenderPipelineCore/Private/Graph/MovieGraphDefaultAudioRenderer.cpp
+++ b/Engine/Plugins/MovieScene/MovieRenderPipeline/Source/MovieRenderPipelineCore/Private/Graph/MovieGraphDefaultAudioRenderer.cpp
@@ -150,7 +150,7 @@ void UMovieGraphDefaultAudioRenderer::SetupAudioRendering()
 		AudioState.PrevRenderEveryTickValue = AudioRenderEveryTickCvar->GetInt();
 		
 		// Override it to prevent it from automatically ticking, we'll control this below
-		AudioRenderEveryTickCvar->SetWithCurrentPriority(0);
+		AudioRenderEveryTickCvar->SetWithCurrentPriority(0, NAME_None, ECVF_SetByConsole, ECVF_SetByScalability);
 	}
 
 	// Ensure that the NRT audio doesn't get muted
@@ -158,7 +158,7 @@ void UMovieGraphDefaultAudioRenderer::SetupAudioRendering()
 	{
 		AudioState.PrevNeverMuteNRTAudioValue = NeverMuteNRTAudioCvar->GetInt();
 		
-		NeverMuteNRTAudioCvar->SetWithCurrentPriority(1);
+		NeverMuteNRTAudioCvar->SetWithCurrentPriority(1, NAME_None, ECVF_SetByConsole, ECVF_SetByScalability);
 	}
 }
 
diff --git a/Engine/Plugins/MovieScene/MovieRenderPipeline/Source/MovieRenderPipelineCore/Public/MoviePipelineUtils.h b/Engine/Plugins/MovieScene/MovieRenderPipeline/Source/MovieRenderPipelineCore/Public/MoviePipelineUtils.h
--- a/Engine/Plugins/MovieScene/MovieRenderPipeline/Source/MovieRenderPipelineCore/Public/MoviePipelineUtils.h
+++ b/Engine/Plugins/MovieScene/MovieRenderPipeline/Source/MovieRenderPipelineCore/Public/MoviePipelineUtils.h
@@ -60,11 +60,11 @@ MOVIERENDERPIPELINECORE_API void GetPassCompositeData(FMoviePipelineMergerOutput
 		if(bUseOverride) \
 		{ \
 			InOutVariable = CVar->GetBool(); \
-			CVar->SetWithCurrentPriority(OverrideValue); \
+			CVar->SetWithCurrentPriority(OverrideValue, NAME_None, ECVF_SetByConsole, ECVF_SetByScalability); \
 		} \
 		else \
 		{ \
-			CVar->SetWithCurrentPriority(InOutVariable); \
+			CVar->SetWithCurrentPriority(InOutVariable, NAME_None, ECVF_SetByConsole, ECVF_SetByScalability); \
 		} \
 	} \
 }
@@ -78,11 +78,11 @@ MOVIERENDERPIPELINECORE_API void GetPassCompositeData(FMoviePipelineMergerOutput
 		if(bUseOverride) \
 		{ \
 			InOutVariable = CVar->GetInt(); \
-			CVar->SetWithCurrentPriority(OverrideValue); \
+			CVar->SetWithCurrentPriority(OverrideValue, NAME_None, ECVF_SetByConsole, ECVF_SetByScalability); \
 		} \
 		else \
 		{ \
-			CVar->SetWithCurrentPriority(InOutVariable); \
+			CVar->SetWithCurrentPriority(InOutVariable, NAME_None, ECVF_SetByConsole, ECVF_SetByScalability); \
 		} \
 	} \
 }
@@ -96,11 +96,11 @@ MOVIERENDERPIPELINECORE_API void GetPassCompositeData(FMoviePipelineMergerOutput
 		if(bUseOverride) \
 		{ \
 			InOutVariable = CVar->GetFloat(); \
-			CVar->SetWithCurrentPriority(OverrideValue); \
+			CVar->SetWithCurrentPriority(OverrideValue, NAME_None, ECVF_SetByConsole, ECVF_SetByScalability); \
 		} \
 		else \
 		{ \
-			CVar->SetWithCurrentPriority(InOutVariable); \
+			CVar->SetWithCurrentPriority(InOutVariable, NAME_None, ECVF_SetByConsole, ECVF_SetByScalability); \
 		} \
 	} \
 }
@@ -114,11 +114,11 @@ MOVIERENDERPIPELINECORE_API void GetPassCompositeData(FMoviePipelineMergerOutput
 		if(bUseOverride) \
 		{ \
 			InOutVariable = CVar->GetBool(); \
-			CVar->SetWithCurrentPriority(OverrideValue); \
+			CVar->SetWithCurrentPriority(OverrideValue, NAME_None, ECVF_SetByConsole, ECVF_SetByScalability); \
 		} \
 		else \
 		{ \
-			CVar->SetWithCurrentPriority(InOutVariable); \
+			CVar->SetWithCurrentPriority(InOutVariable, NAME_None, ECVF_SetByConsole, ECVF_SetByScalability); \
 		} \
 	} \
 }
@@ -132,11 +132,11 @@ MOVIERENDERPIPELINECORE_API void GetPassCompositeData(FMoviePipelineMergerOutput
 		if(bUseOverride) \
 		{ \
 			InOutVariable = CVar->GetInt(); \
-			CVar->SetWithCurrentPriority(OverrideValue); \
+			CVar->SetWithCurrentPriority(OverrideValue, NAME_None, ECVF_SetByConsole, ECVF_SetByScalability); \
 		} \
 		else \
 		{ \
-			CVar->SetWithCurrentPriority(InOutVariable); \
+			CVar->SetWithCurrentPriority(InOutVariable, NAME_None, ECVF_SetByConsole, ECVF_SetByScalability); \
 		} \
 	} \
 }
@@ -150,11 +150,11 @@ MOVIERENDERPIPELINECORE_API void GetPassCompositeData(FMoviePipelineMergerOutput
 		if(bUseOverride) \
 		{ \
 			InOutVariable = CVar->GetFloat(); \
-			CVar->SetWithCurrentPriority(OverrideValue); \
+			CVar->SetWithCurrentPriority(OverrideValue, NAME_None, ECVF_SetByConsole, ECVF_SetByScalability); \
 		} \
 		else \
 		{ \
-			CVar->SetWithCurrentPriority(InOutVariablee); \
+			CVar->SetWithCurrentPriority(InOutVariablee, NAME_None, ECVF_SetByConsole, ECVF_SetByScalability); \
 		} \
 	} \
 }
diff --git a/Engine/Plugins/MovieScene/MovieRenderPipeline/Source/MovieRenderPipelineSettings/Private/MoviePipelineConsoleVariableSetting.cpp b/Engine/Plugins/MovieScene/MovieRenderPipeline/Source/MovieRenderPipelineSettings/Private/MoviePipelineConsoleVariableSetting.cpp
--- a/Engine/Plugins/MovieScene/MovieRenderPipeline/Source/MovieRenderPipelineSettings/Private/MoviePipelineConsoleVariableSetting.cpp
+++ b/Engine/Plugins/MovieScene/MovieRenderPipeline/Source/MovieRenderPipelineSettings/Private/MoviePipelineConsoleVariableSetting.cpp
@@ -34,15 +34,15 @@ namespace UE
 			// cast to integer first (to avoid scientific notation) if we know the cvar is an integer.
 			if (InCVar->IsVariableInt())
 			{
-				InCVar->SetWithCurrentPriority(static_cast<int32>(InValue));
+				InCVar->SetWithCurrentPriority(static_cast<int32>(InValue), NAME_None, ECVF_SetByConsole, ECVF_SetByScalability);
 			}
 			else if (InCVar->IsVariableBool())
 			{
-				InCVar->SetWithCurrentPriority(InValue != 0.f ? true : false);
+				InCVar->SetWithCurrentPriority(InValue != 0.f ? true : false, NAME_None, ECVF_SetByConsole, ECVF_SetByScalability);
 			}
 			else
 			{
-				InCVar->SetWithCurrentPriority(InValue);
+				InCVar->SetWithCurrentPriority(InValue, NAME_None, ECVF_SetByConsole, ECVF_SetByScalability);
 			}
 		}
 
diff --git a/Engine/Source/Runtime/Core/Private/HAL/ConsoleManager.cpp b/Engine/Source/Runtime/Core/Private/HAL/ConsoleManager.cpp
--- a/Engine/Source/Runtime/Core/Private/HAL/ConsoleManager.cpp
+++ b/Engine/Source/Runtime/Core/Private/HAL/ConsoleManager.cpp
@@ -759,18 +759,20 @@ public:
 	bool bHasTaggedArrayData = false;
 
 	void FindHighestPriorityAndTag(
+		EConsoleVariableFlags MinPriority,
 		EConsoleVariableFlags MaxPriority,
 		EConsoleVariableFlags& OutFlags, 
 		FName& OutTag)
 	{
 		OutTag = NAME_None;
-		OutFlags = SETBY_ERROR;
+		OutFlags = MinPriority;
 
 		for (typename TSortedMap<int, TArray<FTaggedHistoryData>>::TReverseIterator It(History); It; ++It)
 		{
 			if (It.Value().Num() > 0)
 			{
-				if ((EConsoleVariableFlags)It.Key() <= MaxPriority)
+				if (((EConsoleVariableFlags)It.Key() <= MaxPriority) &&
+					((EConsoleVariableFlags)It.Key() >= MinPriority))
 				{
 					OutTag = It.Value().Last().Key;
 					OutFlags = (EConsoleVariableFlags)(It.Key());
@@ -978,6 +980,7 @@ public:
 			}
 
 			PriorityHistory->FindHighestPriorityAndTag(
+				Context.MinPriority,
 				Context.MaxPriority,
 				HighestSetBy,
 				HighestTag);
@@ -989,7 +992,12 @@ public:
 			Resolved.Flags = (EConsoleVariableFlags)(HighestSetBy | ExtraFlags);
 			if (HighestSetBy == SETBY_ERROR)
 			{
-				UE_LOG(LogConsoleManager, Display, TEXT("Trying to Replace Cvar=%s Set By Constructor Value. This is not allowed."), *IConsoleManager::Get().FindConsoleObjectName(this));
+				UE_LOG(LogConsoleManager, Display, 
+					TEXT("Trying to Replace Cvar=%s Set By Constructor Value implicitly by calling SetWithCurrentPriority or ReplaceCurrentPriorityAndTag.\n")
+					TEXT("This is not allowed, possible solutions:\n")
+					TEXT("1. Use SetWithCurrentPriority or ReplaceCurrentPriorityAndTag with explicit MinPriority and MaxPriority values, with MinPriority >= ECVF_SetByScalability.\n")
+					TEXT("2. Use explicit Set, Cvar->Set(<Value>, <SetByPriority>)"),
+					*IConsoleManager::Get().FindConsoleObjectName(this));
 			}
 		}
 
diff --git a/Engine/Source/Runtime/Core/Public/HAL/IConsoleManager.h b/Engine/Source/Runtime/Core/Public/HAL/IConsoleManager.h
--- a/Engine/Source/Runtime/Core/Public/HAL/IConsoleManager.h
+++ b/Engine/Source/Runtime/Core/Public/HAL/IConsoleManager.h
@@ -602,6 +602,7 @@ public:
 		ETagMode TagMode = ETagMode::Set;
 
 		EConsoleVariableFlags MaxPriority = ECVF_SetByConsole;
+		EConsoleVariableFlags MinPriority = SETBY_ERROR;
 	};
 
 	virtual FResolvedContext ResolveContext(const FSetContext& Context) = 0;
@@ -755,11 +756,14 @@ public:
 	void SetWithCurrentPriority(
 		T Value, 
 		FName Tag = NAME_None, 
-		EConsoleVariableFlags MaxPriority = ECVF_SetByConsole)
+		EConsoleVariableFlags MaxPriority = ECVF_SetByConsole,
+		EConsoleVariableFlags MinPriority = SETBY_ERROR)
 	{
+		checkf(MinPriority <= MaxPriority, TEXT("MinPriority must be <= MaxPriority"))
 		FSetContext Context;
 		Context.Tag = Tag;
 		Context.PriorityMode = FSetContext::EPriorityMode::ReplaceCurrent;
+		Context.MinPriority = MinPriority;
 		Context.MaxPriority = MaxPriority;
 		EConsoleVariableFlags CurFlags = (EConsoleVariableFlags)(GetFlags() & ECVF_SetByMask);
 		Context.Flags = CurFlags;
@@ -769,12 +773,14 @@ public:
 	template<typename T>
 	void ReplaceCurrentPriorityAndTag(
 		T InValue, 
-		EConsoleVariableFlags MaxPriority = ECVF_SetByConsole
-		)
+		EConsoleVariableFlags MaxPriority = ECVF_SetByConsole,
+		EConsoleVariableFlags MinPriority = SETBY_ERROR)
 	{
+		checkf(MinPriority <= MaxPriority, TEXT("MinPriority must be <= MaxPriority"))
 		FSetContext Context;
 		Context.TagMode = FSetContext::ETagMode::ReplaceCurrent;
 		Context.PriorityMode = FSetContext::EPriorityMode::ReplaceCurrent;
+		Context.MinPriority = MinPriority;
 		Context.MaxPriority = MaxPriority;
 		EConsoleVariableFlags CurFlags = (EConsoleVariableFlags)(GetFlags() & ECVF_SetByMask);
 		Context.Flags = CurFlags;
diff --git a/Engine/Source/Runtime/Engine/Private/Scalability.cpp b/Engine/Source/Runtime/Engine/Private/Scalability.cpp
--- a/Engine/Source/Runtime/Engine/Private/Scalability.cpp
+++ b/Engine/Source/Runtime/Engine/Private/Scalability.cpp
@@ -432,7 +432,8 @@ static void SetGroupQualityLevel(const TCHAR* InGroupName, int32 InQualityLevel,
 						if (IConsoleVariable* CVar = IConsoleManager::Get().FindConsoleVariable(*CVarName))
 						{
 							EConsoleVariableFlags MaxPriority = ECVF_SetByPreview;
-							CVar->ReplaceCurrentPriorityAndTag(*CVarValue, MaxPriority);
+							EConsoleVariableFlags MinPriority = ECVF_SetByScalability;
+							CVar->ReplaceCurrentPriorityAndTag(*CVarValue, MaxPriority, MinPriority);
 						}
 					}
 				}
@@ -893,7 +894,8 @@ template<typename T> void SetQualityLevelCVar(TAutoConsoleVariable<T>& CVarToEdi
 	else if (GEditor && GEditor->IsFeatureLevelPreviewActive())
 	{
 		EConsoleVariableFlags MaxPriority = ECVF_SetByPreview;
-		CVarToEdit.AsVariable()->ReplaceCurrentPriorityAndTag(DesiredValue, MaxPriority);
+		EConsoleVariableFlags MinPriority = ECVF_SetByScalability;
+		CVarToEdit.AsVariable()->ReplaceCurrentPriorityAndTag(DesiredValue, MaxPriority, MinPriority);
 	}
 #endif
 	else
