MoviePipeline: Disable window input in Local Execution. Related to crash fix for Remote Execution from CL 45851164

#rnx

diff --git a/Engine/Plugins/MovieScene/MovieRenderPipeline/Source/MovieRenderPipelineEditor/Private/MoviePipelinePIEExecutor.cpp b/Engine/Plugins/MovieScene/MovieRenderPipeline/Source/MovieRenderPipelineEditor/Private/MoviePipelinePIEExecutor.cpp
--- a/Engine/Plugins/MovieScene/MovieRenderPipeline/Source/MovieRenderPipelineEditor/Private/MoviePipelinePIEExecutor.cpp
+++ b/Engine/Plugins/MovieScene/MovieRenderPipeline/Source/MovieRenderPipelineEditor/Private/MoviePipelinePIEExecutor.cpp
@@ -173,6 +173,8 @@ void UMoviePipelinePIEExecutor::Start(const UMoviePipelineExecutorJob* InJob)
 		.SupportsMinimize(true)
 		.SizingRule(ESizingRule::UserSized);
 
+	CustomWindow->SetAcceptsInput(false);
+
 	WeakCustomWindow = CustomWindow;
 	FSlateApplication::Get().AddWindow(CustomWindow, !IsRenderingOffscreen());
 
