MoviePipeline: Fixed potential out-of-order frame writes for movie files, which can occur in rare edge-cases like when using animated overscan.

diff --git a/Engine/Plugins/MovieScene/MovieRenderPipeline/Source/MovieRenderPipelineCore/Private/Graph/MovieGraphOutputMerger.cpp b/Engine/Plugins/MovieScene/MovieRenderPipeline/Source/MovieRenderPipelineCore/Private/Graph/MovieGraphOutputMerger.cpp
--- a/Engine/Plugins/MovieScene/MovieRenderPipeline/Source/MovieRenderPipelineCore/Private/Graph/MovieGraphOutputMerger.cpp
+++ b/Engine/Plugins/MovieScene/MovieRenderPipeline/Source/MovieRenderPipelineCore/Private/Graph/MovieGraphOutputMerger.cpp
@@ -147,8 +147,8 @@ namespace UE::MovieGraph
 		OutputFrame->ImageOutputData.FindOrAdd(NewLayerId) = MoveTemp(InData);
 		
 		// Check to see if this was the last piece of data needed for this frame.
-		int32 TotalPasses = OutputFrame->ExpectedRenderPasses.Num();
-		int32 FinishedPasses = OutputFrame->ImageOutputData.Num();
+		const int32 TotalPasses = OutputFrame->ExpectedRenderPasses.Num();
+		const int32 FinishedPasses = OutputFrame->ImageOutputData.Num();
 		
 		if (FinishedPasses == TotalPasses)
 		{
@@ -171,9 +171,26 @@ namespace UE::MovieGraph
 					return FirstPayload->CompositingSortOrder < SecondPayload->CompositingSortOrder;
 				}
 			);
+		}
+
+		// Frames can arrive out-of-order in some cases. We should process them in order though, which is especially important for movies (since
+		// those write to one file, vs. many individual files for image sequences, so it's important to provide movie writers the frames in order).
+		// Note: PendingData is kept sorted by its key (frame number).
+		TArray<int32> PendingFrameNumbers;
+		PendingData.GetKeys(PendingFrameNumbers);
+		for (const int32 PendingFrameNumber : PendingFrameNumbers)
+		{
+			const FMovieGraphOutputMergerFrame& OutputMergerFrame = PendingData[PendingFrameNumber];
+			const bool bFrameIsComplete = OutputMergerFrame.ExpectedRenderPasses.Num() == OutputMergerFrame.ImageOutputData.Num();
+			if (!bFrameIsComplete)
+			{
+				// This frame hasn't finished yet, so delay pushing it (and any other out-of-order frames) to FinishedFrames until it's finished.
+				break;
+			}
+			
 			// Move this frame into our FinishedFrames array so the Game Thread can read it at its leisure
 			FMovieGraphOutputMergerFrame FinalFrame;
-			ensureMsgf(PendingData.RemoveAndCopyValue(IndexedFrameNumber, FinalFrame), TEXT("Could not find frame in pending data, output will be skipped!"));
+			ensureMsgf(PendingData.RemoveAndCopyValue(PendingFrameNumber, FinalFrame), TEXT("Could not find frame in pending data, output will be skipped!"));
 			
 			// TQueue is thread safe so it's okay to just push the data into it.
 			FinishedFrames.Enqueue(MoveTemp(FinalFrame));
diff --git a/Engine/Plugins/MovieScene/MovieRenderPipeline/Source/MovieRenderPipelineCore/Private/Graph/MovieGraphOutputMerger.h b/Engine/Plugins/MovieScene/MovieRenderPipeline/Source/MovieRenderPipelineCore/Private/Graph/MovieGraphOutputMerger.h
--- a/Engine/Plugins/MovieScene/MovieRenderPipeline/Source/MovieRenderPipelineCore/Private/Graph/MovieGraphOutputMerger.h
+++ b/Engine/Plugins/MovieScene/MovieRenderPipeline/Source/MovieRenderPipelineCore/Private/Graph/MovieGraphOutputMerger.h
@@ -43,8 +43,8 @@ namespace UE::MovieGraph
 		/** The Movie Pipeline that owns us. */
 		TWeakObjectPtr<UMovieGraphPipeline> WeakMoviePipeline;
 
-		/** Data that is expected but not fully available yet. */
-		TMap<int32, FMovieGraphOutputMergerFrame> PendingData;
+		/** Data that is expected but not fully available yet. Sorted by frame number. */
+		TSortedMap<int32, FMovieGraphOutputMergerFrame> PendingData;
 
 		/** Mutex that protects adding/updating/removing from PendingData */
 		FCriticalSection PendingDataMutex;
diff --git a/Engine/Plugins/MovieScene/MovieRenderPipeline/Source/MovieRenderPipelineCore/Private/MoviePipelineOutputBuilder.cpp b/Engine/Plugins/MovieScene/MovieRenderPipeline/Source/MovieRenderPipelineCore/Private/MoviePipelineOutputBuilder.cpp
--- a/Engine/Plugins/MovieScene/MovieRenderPipeline/Source/MovieRenderPipelineCore/Private/MoviePipelineOutputBuilder.cpp
+++ b/Engine/Plugins/MovieScene/MovieRenderPipeline/Source/MovieRenderPipelineCore/Private/MoviePipelineOutputBuilder.cpp
@@ -92,8 +92,8 @@ void FMoviePipelineOutputMerger::OnCompleteRenderPassDataAvailable_AnyThread(TUn
 	OutputFrame->ImageOutputData.FindOrAdd(Payload->PassIdentifier) = MoveTemp(InData);
 	
 	// Check to see if this was the last piece of data needed for this frame.
-	int32 TotalPasses = OutputFrame->ExpectedRenderPasses.Num();
-	int32 SucceededPasses = OutputFrame->ImageOutputData.Num();
+	const int32 TotalPasses = OutputFrame->ExpectedRenderPasses.Num();
+	const int32 SucceededPasses = OutputFrame->ImageOutputData.Num();
 
 	if (SucceededPasses == TotalPasses)
 	{
@@ -107,12 +107,28 @@ void FMoviePipelineOutputMerger::OnCompleteRenderPassDataAvailable_AnyThread(TUn
 					return FirstPayload->SortingOrder < SecondPayload->SortingOrder;
 				}
 		);
+	}
+
+	// Frames can arrive out-of-order in some cases. We should process them in order though, which is especially important for movies (since
+	// those write to one file, vs. many individual files for image sequences, so it's important to provide movie writers the frames in order).
+	// Note: PendingData is kept sorted by frame number.
+	TArray<FMoviePipelineFrameOutputState> PendingFrameKeys;
+	PendingData.GetKeys(PendingFrameKeys);
+	for (const FMoviePipelineFrameOutputState& PendingFrameKey : PendingFrameKeys)
+	{
+		const FMoviePipelineMergerOutputFrame& OutputMergerFrame = PendingData[PendingFrameKey];
+		const bool bFrameIsComplete = OutputMergerFrame.ExpectedRenderPasses.Num() == OutputMergerFrame.ImageOutputData.Num();
+		if (!bFrameIsComplete)
+		{
+			// This frame hasn't finished yet, so delay pushing it (and any other out-of-order frames) to FinishedFrames until it's finished.
+			break;
+		}
 
 		// Transfer ownership from the map to here; It's important that we use the manually looked up OutputFrame from PendingData
 		// as PendingData uses the equality operator. Some combinations of temporal sampling + slowmo tracks results in different
 		// original source frame numbers, which would cause the tmap lookup to fail and thus returning an empty frame.
 		FMoviePipelineMergerOutputFrame FinalFrame;
-		ensureMsgf(PendingData.RemoveAndCopyValue(OutputFrame->FrameOutputState, FinalFrame), TEXT("Could not find frame in pending data, output will be skipped!"));
+		ensureMsgf(PendingData.RemoveAndCopyValue(PendingFrameKey, FinalFrame), TEXT("Could not find frame in pending data, output will be skipped!"));
 		FinishedFrames.Enqueue(MoveTemp(FinalFrame));
 	}
 }
diff --git a/Engine/Plugins/MovieScene/MovieRenderPipeline/Source/MovieRenderPipelineCore/Public/MoviePipelineOutputBuilder.h b/Engine/Plugins/MovieScene/MovieRenderPipeline/Source/MovieRenderPipelineCore/Public/MoviePipelineOutputBuilder.h
--- a/Engine/Plugins/MovieScene/MovieRenderPipeline/Source/MovieRenderPipelineCore/Public/MoviePipelineOutputBuilder.h
+++ b/Engine/Plugins/MovieScene/MovieRenderPipeline/Source/MovieRenderPipelineCore/Public/MoviePipelineOutputBuilder.h
@@ -34,8 +34,8 @@ private:
 	/** The Movie Pipeline that owns us. */
 	TWeakObjectPtr<UMoviePipeline> WeakMoviePipeline;
 
-	/** Data that is expected but not fully available yet. */
-	TMap<FMoviePipelineFrameOutputState, FMoviePipelineMergerOutputFrame> PendingData;
+	/** Data that is expected but not fully available yet. Sorted by frame number. */
+	TSortedMap<FMoviePipelineFrameOutputState, FMoviePipelineMergerOutputFrame> PendingData;
 
 	/** Mutex that protects adding/updating/removing from ActiveData */
 	FCriticalSection ActiveDataMutex;
diff --git a/Engine/Plugins/MovieScene/MovieRenderPipeline/Source/MovieRenderPipelineCore/Public/MovieRenderPipelineDataTypes.h b/Engine/Plugins/MovieScene/MovieRenderPipeline/Source/MovieRenderPipelineCore/Public/MovieRenderPipelineDataTypes.h
--- a/Engine/Plugins/MovieScene/MovieRenderPipeline/Source/MovieRenderPipelineCore/Public/MovieRenderPipelineDataTypes.h
+++ b/Engine/Plugins/MovieScene/MovieRenderPipeline/Source/MovieRenderPipelineCore/Public/MovieRenderPipelineDataTypes.h
@@ -811,6 +811,11 @@ public:
 		return !(*this == InRHS);
 	}
 
+	bool operator < (const FMoviePipelineFrameOutputState& InRHS) const
+	{
+		return this->OutputFrameNumber < InRHS.OutputFrameNumber;
+	}
+
 	friend uint32 GetTypeHash(FMoviePipelineFrameOutputState OutputState)
 	{
 		return GetTypeHash(OutputState.OutputFrameNumber);
