MoviePipeline: Fixed an issue that caused files to not be written out if a user (incorrectly) includes a space in a Naming Token namespace or token name (eg, {Token Namespace:TestToken}).

diff --git a/Engine/Plugins/MovieScene/MovieRenderPipeline/Source/MovieRenderPipelineCore/Private/Graph/MovieGraphBlueprintLibrary.cpp b/Engine/Plugins/MovieScene/MovieRenderPipeline/Source/MovieRenderPipelineCore/Private/Graph/MovieGraphBlueprintLibrary.cpp
--- a/Engine/Plugins/MovieScene/MovieRenderPipeline/Source/MovieRenderPipelineCore/Private/Graph/MovieGraphBlueprintLibrary.cpp
+++ b/Engine/Plugins/MovieScene/MovieRenderPipeline/Source/MovieRenderPipelineCore/Private/Graph/MovieGraphBlueprintLibrary.cpp
@@ -24,6 +24,29 @@
 
 namespace
 {
+	/**
+	 * Similar to UE::NamingTokens::Utils::GetTokenKeysFromString(), but allows for spaces within the namespace or token name.
+	 * 
+	 * Spaces are not allowed by Naming Tokens, but users may accidentally place them in format strings. We should still detect these
+	 * as valid tokens, but keep them un-resolved.
+	 */
+	TArray<FString> GetTokensFromString(const FString& InFormatString)
+	{
+    	static const FString PatternString = FString::Printf(TEXT(R"(\{\s*((?:[a-zA-Z0-9_ ]+%s)*[a-zA-Z0-9_ ]+)\s*\})"), *UE::NamingTokens::Utils::GetNamespaceDelimiter());
+    	static const FRegexPattern Pattern(PatternString);
+    
+    	TSet<FString, FLocKeySetFuncs> Tokens;
+    	FRegexMatcher Matcher(Pattern, InFormatString);
+
+    	while (Matcher.FindNext())
+    	{
+    		FString Token = Matcher.GetCaptureGroup(1);
+    		Tokens.Add(Token);
+    	}
+    
+    	return Tokens.Array();
+	}
+	
 	/**
 	 * From the provided format string, resolve all of the non-MRG {tokens} from the Naming Tokens system. Resolved values are provided in both the
 	 * InOutFilenameArguments and InOutFormatArguments maps. The former is exposed to the API, and the latter is just used for FString::Format().
@@ -54,10 +77,19 @@ namespace
 			InOutFormatArguments.Add(TempMrgTokenName, EvaluatedTokenValue);
 			InOutFilenameArguments.Add(InTokenName, EvaluatedTokenValue);
 		};
+
+		// Replaces any invalid characters (like ':') in a token with '_'. Colons in tokens can happen if the token provided is invalid and
+		// fails to be resolved.
+		auto SanitizeInvalidToken = [&InOutFormatString](const FString& InTokenName)
+		{
+			InOutFormatString.ReplaceInline(
+				*FString::Printf(TEXT("{%s}"), *InTokenName),
+				*FString::Printf(TEXT("{%s}"), *FPaths::MakeValidFileName(InTokenName, TEXT('_'))));
+		};
 		
 		// Parse out all of the filename argument {tokens} in the string, and find the ones that do not come from MRG.
-		TArray<FString> NonMrgTokenNames;
-		for (const FString& FilenameArgument : UE::NamingTokens::Utils::GetTokenKeysFromString(InOutFormatString))
+		TSet<FString> NonMrgTokenNames;
+		for (const FString& FilenameArgument : GetTokensFromString(InOutFormatString))
 		{
 			if (!InOutFilenameArguments.Contains(FilenameArgument))
 			{
@@ -74,15 +106,24 @@ namespace
 		}
 
 		// Resolve via the Naming Tokens system and inject the values into the argument maps.
-		for (const FNamingTokenValueData& EvaluatedTokenData : NamingTokensSubsystem->EvaluateTokenList(NonMrgTokenNames))
+		TSet<FString> ValidNamingTokens;
+		for (const FNamingTokenValueData& EvaluatedTokenData : NamingTokensSubsystem->EvaluateTokenList(NonMrgTokenNames.Array()))
 		{
+			const FString NamespacedToken = UE::NamingTokens::Utils::CombineNamespaceAndTokenKey(EvaluatedTokenData.TokenNamespace, EvaluatedTokenData.TokenKey);
+
+			// If a token was returned by EvaluateTokenList(), it was recognized as meeting token naming requirements by the Naming Tokens
+			// system. However, that does not guarantee that the token was actually registered or can be resolved.
+			ValidNamingTokens.Add(NamespacedToken);
+			
 			// If the token could not be resolved, skip it
 			if (!EvaluatedTokenData.bWasEvaluated)
 			{
+				// Only namespaced tokens need to be sanitized. The colon delimiter in namespaced tokens is the primary character that
+				// needs to be removed.
+				SanitizeInvalidToken(NamespacedToken);
 				continue;
 			}
-
-			const FString NamespacedToken = UE::NamingTokens::Utils::CombineNamespaceAndTokenKey(EvaluatedTokenData.TokenNamespace, EvaluatedTokenData.TokenKey);
+			
 			if (NonMrgTokenNames.Contains(NamespacedToken))
 			{
 				UpdateFormatStringAndArguments(NamespacedToken, EvaluatedTokenData);
@@ -94,6 +135,15 @@ namespace
 				UpdateFormatStringAndArguments(NonNamespacedToken, EvaluatedTokenData);
 			}
 		}
+
+		// For the tokens that weren't recognized as MRG tokens or valid Naming Tokens, keep them as-is in the format string. However, if it's a
+		// namespaced Naming Token, replace the colon delimiter with an underscore so the file path remains valid. Tokens that can fall into this
+		// code path are typically those that have an invalid space in their name (eg, "Token Namespace:TokenName").
+		const TSet<FString> UnhandledTokens = NonMrgTokenNames.Difference(ValidNamingTokens);
+		for (const FString& UnhandledToken : UnhandledTokens)
+		{
+			SanitizeInvalidToken(UnhandledToken);
+		}
 	}
 
 	FString ResolveFormatArgumentsHelper(const FString& InFormatString, const FMovieGraphFilenameResolveParams& InParams, const bool bResolveToPath, FMovieGraphResolveArgs& OutMergedFormatArgs)
