Fix MoviePipelineUtils to use correct type for cloth DynamicSubstepDeltaTime.

diff --git a/Engine/Plugins/MovieScene/MovieRenderPipeline/Source/MovieRenderPipelineCore/Private/MoviePipelineUtils.cpp b/Engine/Plugins/MovieScene/MovieRenderPipeline/Source/MovieRenderPipelineCore/Private/MoviePipelineUtils.cpp
--- a/Engine/Plugins/MovieScene/MovieRenderPipeline/Source/MovieRenderPipelineCore/Private/MoviePipelineUtils.cpp
+++ b/Engine/Plugins/MovieScene/MovieRenderPipeline/Source/MovieRenderPipelineCore/Private/MoviePipelineUtils.cpp
@@ -1320,7 +1320,7 @@ namespace UE
 								{
 									constexpr int32 MinNumSubsteps = 1;
 									const int32 NumSubsteps = FMath::Max(ClothAssetInteractor->GetIntPropertyValue(TEXT("NumSubsteps"), LODIndex, MinNumSubsteps), MinNumSubsteps);
-									const int32 DynamicSubstepDeltaTime = ClothAssetInteractor->GetFloatPropertyValue(TEXT("DynamicSubstepDeltaTime"), LODIndex, MinNumSubsteps);
+									const float DynamicSubstepDeltaTime = ClothAssetInteractor->GetFloatPropertyValue(TEXT("DynamicSubstepDeltaTime"), LODIndex, MinNumSubsteps);
 									(*ExistingCacheEntry)[LODIndex].NumSubSteps = NumSubsteps;
 									(*ExistingCacheEntry)[LODIndex].DynamicSubstepDeltaTime = DynamicSubstepDeltaTime;
 								}
@@ -1356,7 +1356,7 @@ namespace UE
 						for (int32 LODIndex = 0; LODIndex < Pair.Value.Num(); ++LODIndex)
 						{
 							ClothAssetInteractor->SetIntPropertyValue(TEXT("NumSubsteps"), LODIndex, Pair.Value[LODIndex].NumSubSteps);
-							ClothAssetInteractor->SetIntPropertyValue(TEXT("DynamicSubstepDeltaTime"), LODIndex, Pair.Value[LODIndex].DynamicSubstepDeltaTime);
+							ClothAssetInteractor->SetFloatPropertyValue(TEXT("DynamicSubstepDeltaTime"), LODIndex, Pair.Value[LODIndex].DynamicSubstepDeltaTime);
 						}
 					}
 				}
