MoviePipeline: MRQ/MRG will now prohibit Slate from throttling during a render. If Slate throttles, it may cause the world to not tick, leading to hitches/pauses in the render.

diff --git a/Engine/Plugins/MovieScene/MovieRenderPipeline/Source/MovieRenderPipelineCore/Private/Graph/Nodes/MovieGraphGlobalGameOverrides.cpp b/Engine/Plugins/MovieScene/MovieRenderPipeline/Source/MovieRenderPipelineCore/Private/Graph/Nodes/MovieGraphGlobalGameOverrides.cpp
--- a/Engine/Plugins/MovieScene/MovieRenderPipeline/Source/MovieRenderPipelineCore/Private/Graph/Nodes/MovieGraphGlobalGameOverrides.cpp
+++ b/Engine/Plugins/MovieScene/MovieRenderPipeline/Source/MovieRenderPipelineCore/Private/Graph/Nodes/MovieGraphGlobalGameOverrides.cpp
@@ -100,6 +100,7 @@ void UMovieGraphGlobalGameOverridesNode::BuildNewProcessCommandLineArgsImpl(TArr
 	InOutDeviceProfileCvars.AddUnique(FString::Printf(TEXT("r.SkipRedundantTransformUpdate=%d"), 0));
 	InOutDeviceProfileCvars.AddUnique(FString::Printf(TEXT("p.ChaosCloth.UseTimeStepSmoothing=%d"), 0));
 	InOutDeviceProfileCvars.AddUnique(FString::Printf(TEXT("r.Water.SkipWaterInfoTextureRenderWhenWorldRenderingDisabled=%d"), 0));
+	InOutDeviceProfileCvars.AddUnique(FString::Printf(TEXT("Slate.bAllowThrottling=%d"), 0));
 
 	IConsoleVariable* VTInvalidateCvar = IConsoleManager::Get().FindConsoleVariable(TEXT("MoviePipeline.EnableVTInvalidateOnNaniteLOD"));
 	if (VTInvalidateCvar && VTInvalidateCvar->GetInt() != 0)
@@ -230,6 +231,10 @@ void UMovieGraphGlobalGameOverridesNode::ApplySettings(const bool bOverrideValue
 		// Water skips water info texture when the world's game viewport rendering is disabled so we need to prevent this from happening.
 		MOVIEPIPELINE_STORE_AND_OVERRIDE_CVAR_INT_IF_EXIST(NodeCDO->PreviousSkipWaterInfoTextureRenderWhenWorldRenderingDisabled, TEXT("r.Water.SkipWaterInfoTextureRenderWhenWorldRenderingDisabled"), 0, bOverrideValues);
 
+		// Slate throttling can cause the world to not tick, leading to hitching/pausing in the rendered frames. Notifications/toast popping up can trigger
+		// throttling, so force-disable it during a render.
+		MOVIEPIPELINE_STORE_AND_OVERRIDE_CVAR_INT(NodeCDO->PreviousAllowSlateThrottling, TEXT("Slate.bAllowThrottling"), 0, bOverrideValues);
+
 		// This is only a temporary cvar while it's experimental so it's not exposed to the UI, but exposed as a cvar
 		// so that users can turn it off in the event that it causes issues.
 		IConsoleVariable* VTInvalidateCvar = IConsoleManager::Get().FindConsoleVariable(TEXT("MoviePipeline.EnableVTInvalidateOnNaniteLOD"));
diff --git a/Engine/Plugins/MovieScene/MovieRenderPipeline/Source/MovieRenderPipelineCore/Private/MoviePipelineGameOverrideSetting.cpp b/Engine/Plugins/MovieScene/MovieRenderPipeline/Source/MovieRenderPipelineCore/Private/MoviePipelineGameOverrideSetting.cpp
--- a/Engine/Plugins/MovieScene/MovieRenderPipeline/Source/MovieRenderPipelineCore/Private/MoviePipelineGameOverrideSetting.cpp
+++ b/Engine/Plugins/MovieScene/MovieRenderPipeline/Source/MovieRenderPipelineCore/Private/MoviePipelineGameOverrideSetting.cpp
@@ -168,6 +168,10 @@ void UMoviePipelineGameOverrideSetting::ApplyCVarSettings(const bool bOverrideVa
 	// Cloth's time step smoothing messes up the change in number of simulation substeps that fixes the cloth simulation behavior when using Temporal Samples.
 	MOVIEPIPELINE_STORE_AND_OVERRIDE_CVAR_INT(PreviousChaosClothUseTimeStepSmoothing, TEXT("p.ChaosCloth.UseTimeStepSmoothing"), 0, bOverrideValues);
 
+	// Slate throttling can cause the world to not tick, leading to hitching/pausing in the rendered frames. Notifications/toast popping up can trigger
+	// throttling, so force-disable it during a render.
+	MOVIEPIPELINE_STORE_AND_OVERRIDE_CVAR_INT(PreviousAllowSlateThrottling, TEXT("Slate.bAllowThrottling"), 0, bOverrideValues);
+
 	// Water skips water info texture when the world's game viewport rendering is disabled so we need to prevent this from happening.
 	MOVIEPIPELINE_STORE_AND_OVERRIDE_CVAR_INT_IF_EXIST(PreviousSkipWaterInfoTextureRenderWhenWorldRenderingDisabled, TEXT("r.Water.SkipWaterInfoTextureRenderWhenWorldRenderingDisabled"), 0, bOverrideValues);
 
@@ -293,6 +297,7 @@ void UMoviePipelineGameOverrideSetting::BuildNewProcessCommandLineArgsImpl(TArra
 	InOutDeviceProfileCvars.AddUnique(FString::Printf(TEXT("r.SkipRedundantTransformUpdate=%d"), 0));
 	InOutDeviceProfileCvars.AddUnique(FString::Printf(TEXT("p.ChaosCloth.UseTimeStepSmoothing=%d"), 0));
 	InOutDeviceProfileCvars.AddUnique(FString::Printf(TEXT("r.Water.SkipWaterInfoTextureRenderWhenWorldRenderingDisabled=%d"), 0));
+	InOutDeviceProfileCvars.AddUnique(FString::Printf(TEXT("Slate.bAllowThrottling=%d"), 0));
 
 	if (CVarMoviePipelineVTNaniteAutoLOD.GetValueOnGameThread())
 	{
diff --git a/Engine/Plugins/MovieScene/MovieRenderPipeline/Source/MovieRenderPipelineCore/Public/Graph/Nodes/MovieGraphGlobalGameOverrides.h b/Engine/Plugins/MovieScene/MovieRenderPipeline/Source/MovieRenderPipelineCore/Public/Graph/Nodes/MovieGraphGlobalGameOverrides.h
--- a/Engine/Plugins/MovieScene/MovieRenderPipeline/Source/MovieRenderPipelineCore/Public/Graph/Nodes/MovieGraphGlobalGameOverrides.h
+++ b/Engine/Plugins/MovieScene/MovieRenderPipeline/Source/MovieRenderPipelineCore/Public/Graph/Nodes/MovieGraphGlobalGameOverrides.h
@@ -245,6 +245,7 @@ private:
 	int32 PreviousChaosClothUseTimeStepSmoothing;
 	int32 PreviousSkipWaterInfoTextureRenderWhenWorldRenderingDisabled;
 	int32 PreviousNaniteVSMInvalidateOnLODDelta;
+	int32 PreviousAllowSlateThrottling;
 	bool PreviousAlphaOutput;
 
 #if WITH_EDITOR
diff --git a/Engine/Plugins/MovieScene/MovieRenderPipeline/Source/MovieRenderPipelineCore/Public/MoviePipelineGameOverrideSetting.h b/Engine/Plugins/MovieScene/MovieRenderPipeline/Source/MovieRenderPipelineCore/Public/MoviePipelineGameOverrideSetting.h
--- a/Engine/Plugins/MovieScene/MovieRenderPipeline/Source/MovieRenderPipelineCore/Public/MoviePipelineGameOverrideSetting.h
+++ b/Engine/Plugins/MovieScene/MovieRenderPipeline/Source/MovieRenderPipelineCore/Public/MoviePipelineGameOverrideSetting.h
@@ -186,6 +186,7 @@ private:
 	int32 PreviousChaosClothUseTimeStepSmoothing;
 	int32 PreviousSkipWaterInfoTextureRenderWhenWorldRenderingDisabled;
 	int32 PreviousNaniteVSMInvalidateOnLODDelta;
+	int32 PreviousAllowSlateThrottling;
 	bool PreviousAlphaOutput;
 };
 
