MoviePipeline: Fixed an issue in the queue that could cause a job to lose focus when the job selection changes (a regression from 45937938). This change fixes the regression, and also addresses the original reason for 45937938 (ensuring that job variable values are committed properly when job selection state changes).

#rnx

diff --git a/Engine/Plugins/MovieScene/MovieRenderPipeline/Source/MovieRenderPipelineEditor/Private/Widgets/SMoviePipelineQueueEditor.cpp b/Engine/Plugins/MovieScene/MovieRenderPipeline/Source/MovieRenderPipelineEditor/Private/Widgets/SMoviePipelineQueueEditor.cpp
--- a/Engine/Plugins/MovieScene/MovieRenderPipeline/Source/MovieRenderPipelineEditor/Private/Widgets/SMoviePipelineQueueEditor.cpp
+++ b/Engine/Plugins/MovieScene/MovieRenderPipeline/Source/MovieRenderPipelineEditor/Private/Widgets/SMoviePipelineQueueEditor.cpp
@@ -1491,6 +1491,7 @@ void SMoviePipelineQueueEditor::Construct(const FArguments& InArgs)
 
 	TreeView = SNew(STreeView<TSharedPtr<IMoviePipelineQueueTreeItem>>)
 		.TreeItemsSource(&RootNodes)
+		.ClearSelectionOnClick(false)
 		.OnSelectionChanged(this, &SMoviePipelineQueueEditor::OnJobSelectionChanged_Impl)
 		.OnGenerateRow(this, &SMoviePipelineQueueEditor::OnGenerateRow)
 		.OnGetChildren(this, &SMoviePipelineQueueEditor::OnGetChildren)
diff --git a/Engine/Plugins/MovieScene/MovieRenderPipeline/Source/MovieRenderPipelineEditor/Private/Widgets/SMoviePipelineQueuePanel.cpp b/Engine/Plugins/MovieScene/MovieRenderPipeline/Source/MovieRenderPipelineEditor/Private/Widgets/SMoviePipelineQueuePanel.cpp
--- a/Engine/Plugins/MovieScene/MovieRenderPipeline/Source/MovieRenderPipelineEditor/Private/Widgets/SMoviePipelineQueuePanel.cpp
+++ b/Engine/Plugins/MovieScene/MovieRenderPipeline/Source/MovieRenderPipelineEditor/Private/Widgets/SMoviePipelineQueuePanel.cpp
@@ -498,13 +498,6 @@ void SMoviePipelineQueuePanel::OnSelectionChanged(const TArray<UMoviePipelineExe
 		Jobs.Append(InSelectedJobs);
 	}
 
-	if (!JobDetailsPanelWidget->GetSelectedObjects().IsEmpty())
-	{
-		// Force a focus lost event before the selection changes. Depending on how a property was added to the details panel, it may not have time to
-		// properly commit its value if its widget has not lost focus before the selection is changed.
-		FSlateApplication::Get().ClearKeyboardFocus(EFocusCause::SetDirectly);
-	}
-	
 	JobDetailsPanelWidget->SetObjects(Jobs);
 	NumSelectedJobs = InSelectedJobs.Num();
 }
