MoviePipeline: Fixed a bug that could cause small seams in panoramic renders when rendering at very high resolutions.

diff --git a/Engine/Plugins/MovieScene/MovieRenderPipeline/Source/MovieRenderPipelineCore/Public/MoviePipelinePanoramicBlenderBase.cpp b/Engine/Plugins/MovieScene/MovieRenderPipeline/Source/MovieRenderPipelineCore/Public/MoviePipelinePanoramicBlenderBase.cpp
--- a/Engine/Plugins/MovieScene/MovieRenderPipeline/Source/MovieRenderPipelineCore/Public/MoviePipelinePanoramicBlenderBase.cpp
+++ b/Engine/Plugins/MovieScene/MovieRenderPipeline/Source/MovieRenderPipelineCore/Public/MoviePipelinePanoramicBlenderBase.cpp
@@ -243,8 +243,8 @@ void FMoviePipelinePanoramicBlenderBase::BlendSample_AnyThread(TUniquePtr<FImage
 			// Now construct a projection matrix representing the sample matching the original perspective it was taken from.
 			const FMatrix SampleProjectionMatrix = FReversedZPerspectiveMatrix(FMath::DegreesToRadians(SampleHalfHorizontalFoVDegrees), SampleSize.X, SampleSize.Y, Pane.NearClippingPlane);
 			// For our given output size, figure out how many degrees each pixel represents.
-			const float EquiRectMapThetaStep = 360.f / (float)OutputEquirectangularMapSize.X;
-			const float EquiRectMapPhiStep = 180.f / (float)OutputEquirectangularMapSize.Y;
+			const double EquiRectMapThetaStep = 360.0 / (double)OutputEquirectangularMapSize.X;
+			const double EquiRectMapPhiStep = 180.0 / (double)OutputEquirectangularMapSize.Y;
 
 			// Compute the index bounds in the equirectangular map corresponding to the sample bounds, so we don't loop over unnecessary pixels.
 			// This is approximated according to the weighting function for blending too.
@@ -253,10 +253,13 @@ void FMoviePipelinePanoramicBlenderBase::BlendSample_AnyThread(TUniquePtr<FImage
 			// Pitch is clamped, because there is no vertical wrapping in the map
 			// Yaw is not clamped, because horizontal wrapping is possible. 
 			// The MinBound for X can actually be greater than the MaxBound due to wrapping, modulo is applied at eval time to ensure it wraps right.
-			const float SampleYawMin = SampleRotation.Yaw - SampleHalfHorizontalFoVDegrees;
-			const float SampleYawMax = SampleRotation.Yaw + SampleHalfHorizontalFoVDegrees;
-			const int32 PixelIndexHorzMinBound = FMath::FloorToInt(((SampleYawMin)+180.f) / EquiRectMapThetaStep);
-			const int32 PixelIndexHorzMaxBound = FMath::FloorToInt(((SampleYawMax)+180.f) / EquiRectMapThetaStep);
+			// Note that the bounds are computed with doubles. Very high precision is needed here due to the floor operation. With floats, there can
+			// be slight errors at high resolutions that cause floor to evaluate to an incorrect bound (eg, could be 12287.999 which evaluates to
+			// 12287 instead of the 12288). These slight errors in the bounds can lead to seams in the blended result.
+			const double SampleYawMin = SampleRotation.Yaw - SampleHalfHorizontalFoVDegrees;
+			const double SampleYawMax = SampleRotation.Yaw + SampleHalfHorizontalFoVDegrees;
+			const int32 PixelIndexHorzMinBound = FMath::FloorToInt(((SampleYawMin)+180.0) / EquiRectMapThetaStep);
+			const int32 PixelIndexHorzMaxBound = FMath::FloorToInt(((SampleYawMax)+180.0) / EquiRectMapThetaStep);
 
 			const float SamplePitchMin = FMath::Max(SampleRotation.Pitch - SampleHalfVerticalFoVDegrees, -90.f); // Clamped to [-90, 90]
 			const float SamplePitchMax = FMath::Min(SampleRotation.Pitch + SampleHalfVerticalFoVDegrees, 90.f); // Clamped to [-90, 90]
