MoviePipeline: Fixed cvars not being set correctly in some locations within MRG and MRQ after some recent changes to the engine were made. In some scenarios, cvars could be set with ECVF_SetByConstructor priority, which is no longer allowed.

diff --git a/Engine/Plugins/MovieScene/MovieRenderPipeline/Source/MovieRenderPipelineCore/Private/Graph/MovieGraphCVarManager.cpp b/Engine/Plugins/MovieScene/MovieRenderPipeline/Source/MovieRenderPipelineCore/Private/Graph/MovieGraphCVarManager.cpp
--- a/Engine/Plugins/MovieScene/MovieRenderPipeline/Source/MovieRenderPipelineCore/Private/Graph/MovieGraphCVarManager.cpp
+++ b/Engine/Plugins/MovieScene/MovieRenderPipeline/Source/MovieRenderPipelineCore/Private/Graph/MovieGraphCVarManager.cpp
@@ -219,15 +219,15 @@ void FMovieGraphCVarManager::ApplyCVar(IConsoleVariable* InCVar, float InValue)
 	// cast to integer first (to avoid scientific notation) if we know the cvar is an integer.
 	if (InCVar->IsVariableInt())
 	{
-		InCVar->SetWithCurrentPriority(static_cast<int32>(InValue));
+		InCVar->SetWithCurrentPriority(static_cast<int32>(InValue), NAME_None, ECVF_SetByConsole, ECVF_SetByScalability);
 	}
 	else if (InCVar->IsVariableBool())
 	{
-		InCVar->SetWithCurrentPriority(InValue != 0.f ? true : false);
+		InCVar->SetWithCurrentPriority(InValue != 0.f ? true : false, NAME_None, ECVF_SetByConsole, ECVF_SetByScalability);
 	}
 	else
 	{
-		InCVar->SetWithCurrentPriority(InValue);
+		InCVar->SetWithCurrentPriority(InValue, NAME_None, ECVF_SetByConsole, ECVF_SetByScalability);
 	}
 }
 
diff --git a/Engine/Plugins/MovieScene/MovieRenderPipeline/Source/MovieRenderPipelineCore/Private/Graph/MovieGraphDefaultAudioRenderer.cpp b/Engine/Plugins/MovieScene/MovieRenderPipeline/Source/MovieRenderPipelineCore/Private/Graph/MovieGraphDefaultAudioRenderer.cpp
--- a/Engine/Plugins/MovieScene/MovieRenderPipeline/Source/MovieRenderPipelineCore/Private/Graph/MovieGraphDefaultAudioRenderer.cpp
+++ b/Engine/Plugins/MovieScene/MovieRenderPipeline/Source/MovieRenderPipelineCore/Private/Graph/MovieGraphDefaultAudioRenderer.cpp
@@ -172,11 +172,11 @@ void UMovieGraphDefaultAudioRenderer::TeardownAudioRendering() const
 	// This will be null if the NRT wasn't used (module not loaded)
 	if (IConsoleVariable* AudioRenderEveryTickCvar = IConsoleManager::Get().FindConsoleVariable(TEXT("au.nrt.RenderEveryTick")))
 	{
-		AudioRenderEveryTickCvar->Set(AudioState.PrevRenderEveryTickValue, ECVF_SetByConstructor);
+		AudioRenderEveryTickCvar->SetWithCurrentPriority(AudioState.PrevRenderEveryTickValue, NAME_None, ECVF_SetByConsole, ECVF_SetByScalability);
 	}
 
 	if (IConsoleVariable* NeverMuteNRTAudioCvar = IConsoleManager::Get().FindConsoleVariable(TEXT("au.NeverMuteNonRealtimeAudioDevices")))
 	{
-		NeverMuteNRTAudioCvar->Set(AudioState.PrevNeverMuteNRTAudioValue, ECVF_SetByConstructor);
+		NeverMuteNRTAudioCvar->SetWithCurrentPriority(AudioState.PrevNeverMuteNRTAudioValue, NAME_None, ECVF_SetByConsole, ECVF_SetByScalability);
 	}
 }
diff --git a/Engine/Plugins/MovieScene/MovieRenderPipeline/Source/MovieRenderPipelineCore/Private/Graph/MovieGraphDefaultRenderer.cpp b/Engine/Plugins/MovieScene/MovieRenderPipeline/Source/MovieRenderPipelineCore/Private/Graph/MovieGraphDefaultRenderer.cpp
--- a/Engine/Plugins/MovieScene/MovieRenderPipeline/Source/MovieRenderPipelineCore/Private/Graph/MovieGraphDefaultRenderer.cpp
+++ b/Engine/Plugins/MovieScene/MovieRenderPipeline/Source/MovieRenderPipelineCore/Private/Graph/MovieGraphDefaultRenderer.cpp
@@ -368,7 +368,7 @@ void UMovieGraphDefaultRenderer::Render(const FMovieGraphTimeStepData& InTimeSte
 		IConsoleVariable* HairStrandsCVar = IConsoleManager::Get().FindConsoleVariable(TEXT("r.HairStrands.Strands.MotionVectorCheckViewID"));
 		if (HairStrandsCVar)
 		{
-			HairStrandsCVar->SetWithCurrentPriority(0);
+			HairStrandsCVar->SetWithCurrentPriority(0, NAME_None, ECVF_SetByConsole, ECVF_SetByScalability);
 		}
 	}
 
@@ -384,7 +384,7 @@ void UMovieGraphDefaultRenderer::Render(const FMovieGraphTimeStepData& InTimeSte
 		IConsoleVariable* HairStrandsCVar = IConsoleManager::Get().FindConsoleVariable(TEXT("r.HairStrands.Strands.MotionVectorCheckViewID"));
 		if (HairStrandsCVar)
 		{
-			HairStrandsCVar->SetWithCurrentPriority(1);
+			HairStrandsCVar->SetWithCurrentPriority(1, NAME_None, ECVF_SetByConsole, ECVF_SetByScalability);
 		}
 	}
 
diff --git a/Engine/Plugins/MovieScene/MovieRenderPipeline/Source/MovieRenderPipelineCore/Private/MoviePipelineAudioRendering.cpp b/Engine/Plugins/MovieScene/MovieRenderPipeline/Source/MovieRenderPipelineCore/Private/MoviePipelineAudioRendering.cpp
--- a/Engine/Plugins/MovieScene/MovieRenderPipeline/Source/MovieRenderPipelineCore/Private/MoviePipelineAudioRendering.cpp
+++ b/Engine/Plugins/MovieScene/MovieRenderPipeline/Source/MovieRenderPipelineCore/Private/MoviePipelineAudioRendering.cpp
@@ -49,7 +49,7 @@ void UMoviePipeline::SetupAudioRendering()
 	{
 		AudioState.PrevRenderEveryTickValue = AudioRenderEveryTickCvar->GetInt();
 		// Override it to prevent it from automatically ticking, we'll control this below.
-		AudioRenderEveryTickCvar->Set(0, ECVF_SetByConstructor);
+		AudioRenderEveryTickCvar->SetWithCurrentPriority(0, NAME_None, ECVF_SetByConsole, ECVF_SetByScalability);
 	}
 }
 
@@ -64,8 +64,7 @@ void UMoviePipeline::TeardownAudioRendering()
 	// This will be null if the NRT wasn't used (module not loaded).
 	if (AudioRenderEveryTickCvar)
 	{
-		AudioRenderEveryTickCvar->Set(AudioState.PrevRenderEveryTickValue, ECVF_SetByConstructor);
-
+		AudioRenderEveryTickCvar->SetWithCurrentPriority(AudioState.PrevRenderEveryTickValue, NAME_None, ECVF_SetByConsole, ECVF_SetByScalability);
 	}
 }
 
diff --git a/Engine/Plugins/MovieScene/MovieRenderPipeline/Source/MovieRenderPipelineRenderPasses/Private/Graph/Nodes/MovieGraphPathTracerPassNode.cpp b/Engine/Plugins/MovieScene/MovieRenderPipeline/Source/MovieRenderPipelineRenderPasses/Private/Graph/Nodes/MovieGraphPathTracerPassNode.cpp
--- a/Engine/Plugins/MovieScene/MovieRenderPipeline/Source/MovieRenderPipelineRenderPasses/Private/Graph/Nodes/MovieGraphPathTracerPassNode.cpp
+++ b/Engine/Plugins/MovieScene/MovieRenderPipeline/Source/MovieRenderPipelineRenderPasses/Private/Graph/Nodes/MovieGraphPathTracerPassNode.cpp
@@ -89,7 +89,7 @@ void UMovieGraphPathTracerRenderPassNode::SetupImpl(const FMovieGraphRenderPassS
 	if (IConsoleVariable* ProgressDisplayCvar = IConsoleManager::Get().FindConsoleVariable(TEXT("r.PathTracing.ProgressDisplay")))
 	{
 		bOriginalProgressDisplayCvarValue = ProgressDisplayCvar->GetBool();
-		ProgressDisplayCvar->Set(false);
+		ProgressDisplayCvar->SetWithCurrentPriority(false, NAME_None, ECVF_SetByConsole, ECVF_SetByScalability);
 	}
 
 	// Different Path Traced renders can't currently use different NFOR settings, so we take the largest setting from all branches.
@@ -126,14 +126,13 @@ void UMovieGraphPathTracerRenderPassNode::SetupImpl(const FMovieGraphRenderPassS
 	if (IConsoleVariable* FrameCountCvar = IConsoleManager::Get().FindConsoleVariable(TEXT("r.NFOR.FrameCount")))
 	{
 		OriginalFrameCountCvarValue = FrameCountCvar->GetInt();
-		FrameCountCvar->Set(MaxFrameCount);
+		FrameCountCvar->SetWithCurrentPriority(MaxFrameCount, NAME_None, ECVF_SetByConsole, ECVF_SetByScalability);
 	}
 
 	if (IConsoleVariable* DenoiserTypeCvar = IConsoleManager::Get().FindConsoleVariable(TEXT("r.PathTracing.SpatialDenoiser.Type")))
 	{
-
 		OriginalDenoiserType = DenoiserTypeCvar->GetInt();
-		DenoiserTypeCvar->Set(static_cast<int32>(MaxDenoiserType));
+		DenoiserTypeCvar->SetWithCurrentPriority(static_cast<int32>(MaxDenoiserType), NAME_None, ECVF_SetByConsole, ECVF_SetByScalability);
 	}
 
 	//If no temporal denoiser is enabled, provide warning if we use temporal denoiser.
@@ -175,17 +174,17 @@ void UMovieGraphPathTracerRenderPassNode::TeardownImpl()
 	// Restore the original setting for the progress display
 	if (IConsoleVariable* ProgressDisplayCvar = IConsoleManager::Get().FindConsoleVariable(TEXT("r.PathTracing.ProgressDisplay")))
 	{
-		ProgressDisplayCvar->Set(bOriginalProgressDisplayCvarValue);
+		ProgressDisplayCvar->SetWithCurrentPriority(bOriginalProgressDisplayCvarValue, NAME_None, ECVF_SetByConsole, ECVF_SetByScalability);
 	}
 
 	if (IConsoleVariable* FrameCountCvar = IConsoleManager::Get().FindConsoleVariable(TEXT("r.NFOR.FrameCount")))
 	{
-		FrameCountCvar->Set(OriginalFrameCountCvarValue);
+		FrameCountCvar->SetWithCurrentPriority(OriginalFrameCountCvarValue, NAME_None, ECVF_SetByConsole, ECVF_SetByScalability);
 	}
 
 	if (IConsoleVariable* DenoiserTypeCvar = IConsoleManager::Get().FindConsoleVariable(TEXT("r.PathTracing.SpatialDenoiser.Type")))
 	{
-		DenoiserTypeCvar->Set(OriginalDenoiserType);
+		DenoiserTypeCvar->SetWithCurrentPriority(OriginalDenoiserType, NAME_None, ECVF_SetByConsole, ECVF_SetByScalability);
 	}
 }
 
