MoviePipeline: Fixed a bug that caused the multilayer EXR node to generate invalid/empty files if 1) the burn-in is enabled, 2) multipart is enabled, and 3) more than one layer is being rendered.

diff --git a/Engine/Plugins/MovieScene/MovieRenderPipeline/Source/MovieRenderPipelineCore/Private/Graph/Nodes/MovieGraphWidgetRendererBaseNode.cpp b/Engine/Plugins/MovieScene/MovieRenderPipeline/Source/MovieRenderPipelineCore/Private/Graph/Nodes/MovieGraphWidgetRendererBaseNode.cpp
--- a/Engine/Plugins/MovieScene/MovieRenderPipeline/Source/MovieRenderPipelineCore/Private/Graph/Nodes/MovieGraphWidgetRendererBaseNode.cpp
+++ b/Engine/Plugins/MovieScene/MovieRenderPipeline/Source/MovieRenderPipelineCore/Private/Graph/Nodes/MovieGraphWidgetRendererBaseNode.cpp
@@ -9,6 +9,7 @@
 #include "Graph/MovieGraphDataTypes.h"
 #include "Graph/MovieGraphDefaultRenderer.h"
 #include "Graph/MovieGraphPipeline.h"
+#include "Graph/Nodes/MovieGraphBurnInNode.h"
 #include "MovieRenderPipelineCoreModule.h"
 #include "Slate/WidgetRenderer.h"
 #include "Styling/AppStyle.h"
@@ -33,10 +34,31 @@ FString UMovieGraphWidgetRendererBaseNode::GetLayerNameOverride() const
 
 void UMovieGraphWidgetRendererBaseNode::SetupImpl(const FMovieGraphRenderPassSetupData& InSetupData)
 {
+	static const FString MultiLayerExrNodeOutputName = TEXT("MULTI_EXR");
+	bool bIncludedMultiExrBurnIn = false;
+
 	for (const FMovieGraphRenderPassLayerData& LayerData : InSetupData.Layers)
 	{
+		// Temp hack: For multi-layer EXRs that include a burn-in pass, ensure that only one burn-in is included if multiple layers are being
+		// rendered. The current multi-layer EXR node does not have a way to include {tokens} (like {layer_name}) in the Layer Name Format property,
+		// so all burn-in layers will be named the same (and duplicate layer names are not allowed in EXRs). So, for now, restrict the number of
+		// burn-ins to one per multi-layer EXR, which is likely what most users want anyway.
+		if (const UMovieGraphOutputBurnInNode* OutputBurnInNode = Cast<UMovieGraphOutputBurnInNode>(LayerData.RenderPassNode.Get()))
+		{
+			if (OutputBurnInNode->OutputName == MultiLayerExrNodeOutputName)
+			{
+				if (bIncludedMultiExrBurnIn)
+				{
+					continue;
+				}
+
+				bIncludedMultiExrBurnIn = true;
+			}
+		}
+
 		TUniquePtr<FMovieGraphWidgetPass> RendererInstance = GeneratePass();
 		RendererInstance->Setup(InSetupData.Renderer, LayerData);
+
 		CurrentInstances.Add(MoveTemp(RendererInstance));
 	}
 
