MoviePipeline: Fixed a bug that caused MRG to freeze when 1) generating a movie file, and 2) the "Overwrite Existing Output" property on Global Output Settings is turned off.

diff --git a/Engine/Plugins/MovieScene/MovieRenderPipeline/Source/MovieRenderPipelineCore/Private/Graph/Nodes/MovieGraphVideoOutputNode.cpp b/Engine/Plugins/MovieScene/MovieRenderPipeline/Source/MovieRenderPipelineCore/Private/Graph/Nodes/MovieGraphVideoOutputNode.cpp
--- a/Engine/Plugins/MovieScene/MovieRenderPipeline/Source/MovieRenderPipelineCore/Private/Graph/Nodes/MovieGraphVideoOutputNode.cpp
+++ b/Engine/Plugins/MovieScene/MovieRenderPipeline/Source/MovieRenderPipelineCore/Private/Graph/Nodes/MovieGraphVideoOutputNode.cpp
@@ -330,6 +330,11 @@ void UMovieGraphVideoOutputNode::GetOutputFilePaths(const UMovieGraphPipeline* I
 	ResolveParams.FileNameOverride = FileNameFormatString;
 	ResolveParams.Version = Shot ? Shot->ShotInfo.VersionNumber : UMovieGraphBlueprintLibrary::ResolveVersionNumber(ResolveParams);
 	
+	constexpr bool bIncludeRenderPass = false;
+	constexpr bool bTestFrameNumber = false;
+	constexpr bool bIncludeCameraName = false;
+	UE::MoviePipeline::ValidateOutputFormatString(FileNameFormatString, bIncludeRenderPass, bTestFrameNumber, bIncludeCameraName);
+	
 	FMovieGraphResolveArgs FinalFormatArgs;
 	
 	// Strip any frame number tags so we don't get one video file per frame.
