MoviePipeline: Fixed an issue with Quick Render shifting the start/end markers in the active sequence (if there is one) when the Current Viewport mode is used.

Also fixed a related issue where the Current Viewport mode did not respect the active sequence's playhead position (the 1st frame of the sequence was always used).

diff --git a/Engine/Plugins/MovieScene/MovieRenderPipeline/Source/MovieRenderPipelineEditor/Private/Graph/MovieGraphQuickRender.cpp b/Engine/Plugins/MovieScene/MovieRenderPipeline/Source/MovieRenderPipelineEditor/Private/Graph/MovieGraphQuickRender.cpp
--- a/Engine/Plugins/MovieScene/MovieRenderPipeline/Source/MovieRenderPipelineEditor/Private/Graph/MovieGraphQuickRender.cpp
+++ b/Engine/Plugins/MovieScene/MovieRenderPipeline/Source/MovieRenderPipelineEditor/Private/Graph/MovieGraphQuickRender.cpp
@@ -1121,7 +1121,7 @@ void UMovieGraphQuickRenderSubsystem::PerformPreRenderSetup(UWorld* InEditorWorl
 {
 	// Adds the specified camera to the Camera Cut track in the given level sequence.
 	// If there's already a camera cut section within the Camera Cut track, a new one will be added immediately after.
-	// Each section is only one frame long.
+	// The added section is only one frame long unless OptionalFrameRangeOverride is provided.
 	auto AddCameraToLevelSequence = [](ACameraActor* InCameraToAdd, const ULevelSequence* InDestLevelSequence, const TRange<FFrameNumber> OptionalFrameRangeOverride = TRange<FFrameNumber>::Empty())
 	{
 		FMovieGraphTransactionBlocker TransactionBlocker;
@@ -1332,7 +1332,7 @@ void UMovieGraphQuickRenderSubsystem::PerformPreRenderSetup(UWorld* InEditorWorl
 	{
 		// Adding a camera will tweak the length of the utility sequence's playback range
 		ACameraActor* NewCamera = DuplicateViewportCamera();
-		AddCameraToLevelSequence(NewCamera, UtilityLevelSequence.Get());
+		AddCameraToLevelSequence(NewCamera, UtilityLevelSequence.Get(), RenderingMovieScene->GetPlaybackRange());
 
 		RenderingFrameRange = UtilityLevelSequence->GetMovieScene()->GetPlaybackRange();
 	}
