MoviePipeline: Fixed an issue where some Quick Render modes (like "Selected Cameras") would not render with the correct camera transform when the camera is a child within a hierarchy of actors.

diff --git a/Engine/Plugins/MovieScene/MovieRenderPipeline/Source/MovieRenderPipelineEditor/Private/Graph/MovieGraphQuickRender.cpp b/Engine/Plugins/MovieScene/MovieRenderPipeline/Source/MovieRenderPipelineEditor/Private/Graph/MovieGraphQuickRender.cpp
--- a/Engine/Plugins/MovieScene/MovieRenderPipeline/Source/MovieRenderPipelineEditor/Private/Graph/MovieGraphQuickRender.cpp
+++ b/Engine/Plugins/MovieScene/MovieRenderPipeline/Source/MovieRenderPipelineEditor/Private/Graph/MovieGraphQuickRender.cpp
@@ -28,6 +28,7 @@
 #include "Sections/MovieSceneCameraCutSection.h"
 #include "Sections/MovieSceneLevelVisibilitySection.h"
 #include "SequencerUtilities.h"
+#include "Tracks/MovieScene3DTransformTrack.h"
 #include "Tracks/MovieSceneCameraCutTrack.h"
 #include "Tracks/MovieSceneLevelVisibilityTrack.h"
 #include "Tracks/MovieSceneSubTrack.h"
@@ -127,7 +128,90 @@ namespace
 		}
 
 		return CurrentLevelSequence;
-	};
+	}
+
+	/**
+	 * Gets the world-space transform of the given camera, taking into account parent actors in the hierarchy as well as relative transforms
+	 * on the camera component.
+	 */
+	FTransform GetWorldSpaceCameraTransform(const AActor* InActor)
+	{
+		if (!InActor)
+		{
+			UE_LOG(LogMovieRenderPipeline, Warning, TEXT("Invalid camera detected in Quick Render. The render may not appear correct."));
+			return FTransform();
+		}
+
+		// The "camera" being passed in could be a pilot cam, which may not be an actual camera. Thus it won't have a camera component on it, so don't
+		// treat the lack of a camera component as an error.
+		const ACameraActor* CameraActor = Cast<ACameraActor>(InActor);
+		if (!CameraActor)
+		{
+			return InActor->GetTransform();
+		}
+
+		const UCameraComponent* CameraComponent = CameraActor->GetCameraComponent();
+		if (!CameraComponent)
+		{
+			// Use the camera's transform as a fallback
+			return InActor->GetTransform();
+		}
+
+		// The component's transform will most likely be set to relative (vs. world). Calling GetComponentTransform() will take
+		// into account not only the parent actors and the camera actor itself, but the component's transform as well. When the new camera
+		// is spawned, it will inherit the camera component's transform/offset. Therefore, to prevent the camera component's transform
+		// from being applied twice, we need to get the camera actor's position *without* the component's transform baked in (hence the
+		// Inverse() call).
+		const FTransform ComponentRelTransform = CameraComponent->GetRelativeTransform();
+		const FTransform& ComponentWorldTransform = CameraComponent->GetComponentTransform();
+		const FTransform FinalTransform = ComponentRelTransform.Inverse() * ComponentWorldTransform;
+
+		return FinalTransform;
+	}
+
+	/**
+	 * Adds a Transform track for the given camera (with the associated binding). Generally adding a Transform track is not necessary -- the camera
+	 * is spawned in the PIE world, and the transform of the source camera transfers over correctly. However, if the camera has *parent* actors which
+	 * also have transforms, then those will not be taken into account when the camera is spawned. To ensure that the entire hierarchy of transforms
+	 * is accounted for, bake in the world transform of the camera into a Transform track.
+	 */
+	void AddCameraTransformTrack(UMovieScene* InMovieScene, const FGuid& InCameraBindingGuid, const ACameraActor* InCameraToAdd)
+	{
+		UMovieScene3DTransformTrack* TransformTrack = InMovieScene->AddTrack<UMovieScene3DTransformTrack>(InCameraBindingGuid);
+		if (!TransformTrack)
+		{
+			UE_LOG(LogMovieRenderPipeline, Warning, TEXT("Unable to add Transform track for one or more cameras being used in Quick Render. The render may not appear correct."));
+			return;
+		}
+
+		UMovieScene3DTransformSection* TransformSection = Cast<UMovieScene3DTransformSection>(TransformTrack->CreateNewSection());
+		if (!TransformSection)
+		{
+			UE_LOG(LogMovieRenderPipeline, Warning, TEXT("Unable to add transforms for one or more cameras being used in Quick Render. The render may not appear correct."));
+			return;
+		}
+
+		const FTransform FinalTransform = GetWorldSpaceCameraTransform(InCameraToAdd);
+
+		const UE::Math::TVector<double> Location = FinalTransform.GetLocation();
+		const UE::Math::TVector<double> Rotation = FinalTransform.GetRotation().Euler();
+		const UE::Math::TVector<double> Scale = FinalTransform.GetScale3D();
+
+		TransformTrack->AddSection(*TransformSection);
+		TransformSection->SetRange(TRange<FFrameNumber>::All());
+		TransformSection->SetMask(FMovieSceneTransformMask(EMovieSceneTransformChannel::All));
+
+		TArrayView<FMovieSceneDoubleChannel*> DoubleChannels = TransformSection->GetChannelProxy().GetChannels<FMovieSceneDoubleChannel>();
+		DoubleChannels[0]->SetDefault(Location.X);
+		DoubleChannels[1]->SetDefault(Location.Y);
+		DoubleChannels[2]->SetDefault(Location.Z);
+		DoubleChannels[3]->SetDefault(Rotation.X);
+		DoubleChannels[4]->SetDefault(Rotation.Y);
+		DoubleChannels[5]->SetDefault(Rotation.Z);
+		DoubleChannels[6]->SetDefault(Scale.X);
+		DoubleChannels[7]->SetDefault(Scale.Y);
+		DoubleChannels[8]->SetDefault(Scale.Z);
+	}
 }
 
 /**
@@ -1163,6 +1247,9 @@ void UMovieGraphQuickRenderSubsystem::PerformPreRenderSetup(UWorld* InEditorWorl
 		// Add the new camera to the camera cut track.
 		UMovieSceneCameraCutSection* CameraCutSection = CameraCutTrack->AddNewCameraCut(UE::MovieScene::FRelativeObjectBindingID(CameraBinding), StartFrame);
 		CameraCutSection->SetRange(SectionRange);
+
+		// The camera may be parented to other actors within a hierarchy, so we bake in the world transform of the camera.
+		AddCameraTransformTrack(MovieScene, CameraBinding, InCameraToAdd);
 	};
 	
 	// Spawn a new camera which mimics the viewport's camera 
@@ -1187,7 +1274,14 @@ void UMovieGraphQuickRenderSubsystem::PerformPreRenderSetup(UWorld* InEditorWorl
 				// Use the current pilot cam as the template
 				CameraSpawnParams.Template = ViewportPilotActor;
 
-				return Cast<ACameraActor>(PieWorld->SpawnActor(ViewportPilotActor->GetClass(), nullptr, CameraSpawnParams));
+				ACameraActor* SpawnedCamera =
+					Cast<ACameraActor>(PieWorld->SpawnActor(ViewportPilotActor->GetClass(), nullptr, CameraSpawnParams));
+				if (SpawnedCamera)
+				{
+					SpawnedCamera->SetActorTransform(GetWorldSpaceCameraTransform(ViewportPilotActor));
+				}
+
+				return SpawnedCamera;
 			}
 		}
 		
@@ -1211,8 +1305,10 @@ void UMovieGraphQuickRenderSubsystem::PerformPreRenderSetup(UWorld* InEditorWorl
 			if (CachedPrePieData.ViewportActorLock.IsValid())
 			{
 				const TStrongObjectPtr<AActor> LockedActorPin = CachedPrePieData.ViewportActorLock.Pin();
-				CameraPosition = LockedActorPin->GetActorLocation();
-				CameraRotation = LockedActorPin->GetActorRotation();
+
+				const FTransform WorldTransform = GetWorldSpaceCameraTransform(LockedActorPin.Get());
+				CameraPosition = WorldTransform.GetTranslation();
+				CameraRotation = WorldTransform.GetRotation().Rotator();
 			}
 			else
 			{
