MoviePipeline: Fixed an integer overflow in MRQ/MRG when very large (16k+) renders are performed, which could lead to a crash.

diff --git a/Engine/Plugins/MovieScene/MovieRenderPipeline/Source/MovieRenderPipelineCore/Private/MoviePipelineSurfaceReader.cpp b/Engine/Plugins/MovieScene/MovieRenderPipeline/Source/MovieRenderPipelineCore/Private/MoviePipelineSurfaceReader.cpp
--- a/Engine/Plugins/MovieScene/MovieRenderPipeline/Source/MovieRenderPipelineCore/Private/MoviePipelineSurfaceReader.cpp
+++ b/Engine/Plugins/MovieScene/MovieRenderPipeline/Source/MovieRenderPipelineCore/Private/MoviePipelineSurfaceReader.cpp
@@ -218,7 +218,7 @@ void FMoviePipelineSurfaceReader::CopyReadbackTexture_RenderThread(TUniqueFuncti
 		// line by line.
 		if (ExpectedSizeX == ActualSizeX && ExpectedSizeY == ActualSizeY)
 		{
-			FMemory::BigBlockMemcpy(TypeErasedPixels, ColorDataBuffer, (ExpectedSizeX * ExpectedSizeY) * SizeOfColor);
+			FMemory::BigBlockMemcpy(TypeErasedPixels, ColorDataBuffer, (int64)ExpectedSizeX * (int64)ExpectedSizeY * SizeOfColor);
 		}
 		else
 		{
