[MRG] Reference mode for ML data gathering.

Problem: We don't have a direct way to do DumpGPU capture in the MovieRenderGraph atm.

Solutions:
- Add a way to call DumpGPU from the Debug node of the MovieRenderGraph.
- Extend DumpGPU to be able to dump at a path consistent with the rest of the MovieRenderGraph pipeline.

#rnx
#tests win64 editor : generate DumpGPU via MRG on a long sequence

diff --git a/Engine/Plugins/MovieScene/MovieRenderPipeline/Source/MovieRenderPipelineCore/Private/Graph/MovieGraphDefaultRenderer.cpp b/Engine/Plugins/MovieScene/MovieRenderPipeline/Source/MovieRenderPipelineCore/Private/Graph/MovieGraphDefaultRenderer.cpp
--- a/Engine/Plugins/MovieScene/MovieRenderPipeline/Source/MovieRenderPipelineCore/Private/Graph/MovieGraphDefaultRenderer.cpp
+++ b/Engine/Plugins/MovieScene/MovieRenderPipeline/Source/MovieRenderPipelineCore/Private/Graph/MovieGraphDefaultRenderer.cpp
@@ -1,34 +1,34 @@
 // Copyright Epic Games, Inc. All Rights Reserved.
 
 #include "Graph/MovieGraphDefaultRenderer.h"
-#include "Graph/MovieGraphPipeline.h"
+
+#include "DumpGPU.h"
+#include "Engine/TextureRenderTarget2D.h"
 #include "Graph/MovieGraphConfig.h"
+#include "Graph/MovieGraphPipeline.h"
+#include "Graph/MovieGraphUtils.h"
+#include "Graph/Nodes/MovieGraphCameraNode.h"
 #include "Graph/Nodes/MovieGraphDebugNode.h"
 #include "Graph/Nodes/MovieGraphGlobalGameOverrides.h"
 #include "Graph/Nodes/MovieGraphRenderLayerNode.h"
 #include "Graph/Nodes/MovieGraphRenderPassNode.h"
-#include "Graph/Nodes/MovieGraphGlobalGameOverrides.h"
-#include "Graph/Nodes/MovieGraphDebugNode.h"
-#include "Graph/Nodes/MovieGraphCameraNode.h"
+#include "MoviePipelineQueue.h"
+#include "MoviePipelineSurfaceReader.h"
 #include "MovieRenderPipelineCoreModule.h"
+#include "RenderCaptureInterface.h"
 #include "RenderingThread.h"
-#include "Engine/TextureRenderTarget2D.h"
 #include "UObject/Package.h"
-#include "MoviePipelineSurfaceReader.h"
-#include "RenderCaptureInterface.h"
-#include "MoviePipelineQueue.h"
 
 // For flushing async systems
-#include "EngineModule.h"
-#include "MeshCardRepresentation.h"
-#include "ShaderCompiler.h"
-#include "EngineUtils.h"
 #include "AssetCompilingManager.h"
 #include "ContentStreaming.h"
 #include "EngineModule.h"
+#include "EngineUtils.h"
 #include "LandscapeSubsystem.h"
 #include "Materials/MaterialInterface.h"
+#include "MeshCardRepresentation.h"
 #include "RendererInterface.h"
+#include "ShaderCompiler.h"
 
 #include UE_INLINE_GENERATED_CPP_BY_NAME(MovieGraphDefaultRenderer)
 
@@ -334,23 +334,32 @@ void UMovieGraphDefaultRenderer::Render(const FMovieGraphTimeStepData& InTimeSte
 	// so we have to track when any of our *FSceneView* render passes submit stuff (UI renderers don't count)
 	bHasRenderedFirstViewThisFrame = false;
 
-	// Support for RenderDoc captures of just the MRQ work
+	// Support for RenderDoc and DumpGPU captures of just the MRQ work
 #if WITH_EDITOR && !UE_BUILD_SHIPPING
-	TUniquePtr<RenderCaptureInterface::FScopedCapture> ScopedGPUCapture;
+	TUniquePtr<RenderCaptureInterface::FScopedCapture> ScopedRenderDocCapture;
+	TUniquePtr<UE::RenderCore::DumpGPU::FDumpScope> ScopedDumpGPUCapture;
 	{
-		bool bIsRendering = GetOwningGraph()->GetActiveShotList()[GetOwningGraph()->GetCurrentShotIndex()]->ShotInfo.State == EMovieRenderShotState::Rendering;
-		UMovieGraphDebugSettingNode* DebugSettings = InTimeStepData.EvaluatedConfig->GetSettingForBranch<UMovieGraphDebugSettingNode>(UMovieGraphNode::GlobalsPinName);
+	
+		const UMovieGraphDebugSettingNode* DebugSettings = InTimeStepData.EvaluatedConfig->GetSettingForBranch<UMovieGraphDebugSettingNode>(UMovieGraphNode::GlobalsPinName);
+		const bool bIsRendering = GetOwningGraph()->GetActiveShotList()[GetOwningGraph()->GetCurrentShotIndex()]->ShotInfo.State == EMovieRenderShotState::Rendering;
 
-		const bool bWantsDebugCapture = DebugSettings && DebugSettings->bCaptureFramesWithRenderDoc;
-		bool bIsCorrectFrame = false;
-		if (bWantsDebugCapture)
+		const bool bWantsRenderDocDebugCapture = DebugSettings && DebugSettings->bCaptureFramesWithRenderDoc;
+		const bool bWantsDumpGpuDebugCapture = DebugSettings && DebugSettings->bCaptureFramesWithDumpGpu;
+
+		const bool bIsRenderDocCorrectFrame = InTimeStepData.RootFrameNumber == DebugSettings->RenderDocCaptureFrame;
+		const bool bIsDumpGPUCorrectFrame = ((DebugSettings->DumpGpuCaptureFrameStart < 0 || InTimeStepData.RootFrameNumber >= DebugSettings->DumpGpuCaptureFrameStart) &&
+									         (DebugSettings->DumpGpuCaptureFrameStop  < 0 || InTimeStepData.RootFrameNumber <  DebugSettings->DumpGpuCaptureFrameStop));
+
+		if (bIsRendering && bWantsRenderDocDebugCapture && bIsRenderDocCorrectFrame)
 		{
-			bIsCorrectFrame = InTimeStepData.RootFrameNumber == DebugSettings->RenderDocCaptureFrame;
+			ScopedRenderDocCapture = MakeUnique<RenderCaptureInterface::FScopedCapture>(true, *FString::Printf(TEXT("MRQ Frame: %d"), InTimeStepData.RootFrameNumber.Value));
 		}
-
-		if (bIsRendering && bWantsDebugCapture && bIsCorrectFrame)
+		if (bIsRendering && bWantsDumpGpuDebugCapture && bIsDumpGPUCorrectFrame)
 		{
-			ScopedGPUCapture = MakeUnique<RenderCaptureInterface::FScopedCapture>(true, *FString::Printf(TEXT("MRQ Frame: %d"), InTimeStepData.RootFrameNumber.Value));
+			const FString FinalPath = UE::MovieGraph::Capture::PrepareCapturePath(GetOwningGraph(), InTimeStepData.EvaluatedConfig, DebugSettings->DumpGpuCaptureFilePathFormat, TEXT(""), true);
+			const FString Cmd = FString::Printf(TEXT("-noexplorer -path=%s"), *FinalPath);
+			const FString Context = FString::Printf(TEXT("MRQ Frame: %d"), InTimeStepData.RootFrameNumber.Value);
+			ScopedDumpGPUCapture = MakeUnique<UE::RenderCore::DumpGPU::FDumpScope>(true, *Cmd, *Context);
 		}
 	}
 #endif
diff --git a/Engine/Plugins/MovieScene/MovieRenderPipeline/Source/MovieRenderPipelineCore/Private/Graph/MovieGraphPipeline.cpp b/Engine/Plugins/MovieScene/MovieRenderPipeline/Source/MovieRenderPipelineCore/Private/Graph/MovieGraphPipeline.cpp
--- a/Engine/Plugins/MovieScene/MovieRenderPipeline/Source/MovieRenderPipelineCore/Private/Graph/MovieGraphPipeline.cpp
+++ b/Engine/Plugins/MovieScene/MovieRenderPipeline/Source/MovieRenderPipelineCore/Private/Graph/MovieGraphPipeline.cpp
@@ -13,6 +13,7 @@
 #include "Graph/MovieGraphLinearTimeStep.h"
 #include "Graph/MovieGraphOutputMerger.h"
 #include "Graph/MovieGraphRenderLayerSubsystem.h"
+#include "Graph/MovieGraphUtils.h"
 #include "Graph/Nodes/MovieGraphCollectionNode.h"
 #include "Graph/Nodes/MovieGraphDebugNode.h"
 #include "Graph/Nodes/MovieGraphExecuteScriptNode.h"
@@ -98,6 +99,26 @@ void UMovieGraphPipeline::Initialize(UMoviePipelineExecutorJob* InJob, const FMo
 	DuplicateJobAndConfiguration();
 	ExecutePreJobScripts();
 
+	// Create instances of our different classes from the InitConfig
+	GraphRendererInstance = NewObject<UMovieGraphRendererBase>(this, InitConfig.RendererClass);
+	GraphDataSourceInstance = NewObject<UMovieGraphDataSourceBase>(this, InitConfig.DataSourceClass);
+	GraphAudioRendererInstance = NewObject<UMovieGraphAudioRendererBase>(this, InitConfig.AudioRendererClass);
+
+	// Now that we've created our various systems, we will start using them. First thing we do is cache data about
+	// the world, job, player viewport, etc, before we make any modifications. These will be restored at the end
+	// of the render.
+	GraphDataSourceInstance->CacheDataPreJob(InitConfig);
+
+	// Construct the viewport preview UI and bind it to this instance.
+	LoadPreviewWidget();
+
+	GraphAudioRendererInstance->SetupAudioRendering();
+
+	// Update our list of shots from our data source, and then
+	// create our list of active shots, so we don't try to render
+	// a shot the user has deactivated.
+	BuildShotListFromDataSource();
+
 	const FMovieGraphTraversalContext Context = GetCurrentTraversalContext(false);
 	FString OutError;
 	UMovieGraphEvaluatedConfig* FlattenedConfig = GetCurrentJob()->GetGraphPreset()->CreateFlattenedGraph(Context, OutError);
@@ -120,26 +141,6 @@ void UMovieGraphPipeline::Initialize(UMoviePipelineExecutorJob* InJob, const FMo
 		UE_LOG(LogMovieRenderPipeline, Error, TEXT("Unable to initialize Movie Pipeline. The following error(s) were encountered: %s"), *OutError);
 	}
 
-	// Create instances of our different classes from the InitConfig
-	GraphRendererInstance = NewObject<UMovieGraphRendererBase>(this, InitConfig.RendererClass);
-	GraphDataSourceInstance = NewObject<UMovieGraphDataSourceBase>(this, InitConfig.DataSourceClass);
-	GraphAudioRendererInstance = NewObject<UMovieGraphAudioRendererBase>(this, InitConfig.AudioRendererClass);
-
-	// Now that we've created our various systems, we will start using them. First thing we do is cache data about
-	// the world, job, player viewport, etc, before we make any modifications. These will be restored at the end
-	// of the render.
-	GraphDataSourceInstance->CacheDataPreJob(InitConfig);
-
-	// Construct the viewport preview UI and bind it to this instance.
-	LoadPreviewWidget();
-
-	GraphAudioRendererInstance->SetupAudioRendering();
-
-	// Update our list of shots from our data source, and then
-	// create our list of active shots, so we don't try to render
-	// a shot the user has deactivated.
-	BuildShotListFromDataSource();
-
 	//for (ULevelStreaming* Level : GetWorld()->GetStreamingLevels())
 	//{
 	//	UClass* StreamingClass = Level->GetClass();
@@ -787,52 +788,20 @@ void UMovieGraphPipeline::StartUnrealInsightsCapture(UMovieGraphEvaluatedConfig*
 {
 	check(EvaluatedConfig);
 
-	bool bIncludeCDOs = false;
-	bool bExactMatch = true;
-	UMovieGraphDebugSettingNode* DebugSetting =
+	const bool bIncludeCDOs = false;
+	const bool bExactMatch = true;
+	const UMovieGraphDebugSettingNode* DebugSetting =
 		EvaluatedConfig->GetSettingForBranch<UMovieGraphDebugSettingNode>(UMovieGraphSettingNode::GlobalsPinName, bIncludeCDOs, bExactMatch);
+
 	if (!ensureAlwaysMsgf(DebugSetting, TEXT("Failed to find UMovieGraphFileOutputNode. Aborting.")))
 	{
 		return;
 	}
 
-	bIncludeCDOs = true;
-	bExactMatch = true;
-	UMovieGraphGlobalOutputSettingNode* OutputSetting =
-		EvaluatedConfig->GetSettingForBranch<UMovieGraphGlobalOutputSettingNode>(UMovieGraphSettingNode::GlobalsPinName, bIncludeCDOs, bExactMatch);
-	FString FileNameFormatString = OutputSetting->OutputDirectory.Path / DebugSetting->UnrealInsightsTraceFileNameFormat;
-
-	const bool bOverwriteExistingOutput = OutputSetting->bOverwriteExistingOutput;
-
-	// Generate a filename for this encoded file
-	TMap<FString, FString> FormatOverrides;
-	FormatOverrides.Add(TEXT("ext"), TEXT("utrace"));
-
-	FMovieGraphRenderDataIdentifier TempRenderDataIdentifier;
-	TempRenderDataIdentifier.RootBranchName = UMovieGraphSettingNode::GlobalsPinName;
-
-	FMovieGraphFilenameResolveParams Params = FMovieGraphFilenameResolveParams::MakeResolveParams(
-		TempRenderDataIdentifier, this, EvaluatedConfig, GetCurrentTraversalContext(false), FormatOverrides);
-
-	FMovieGraphResolveArgs FinalFormatArgs;
-	constexpr bool bIncludeRenderPass = false;
-	constexpr bool bTestFrameNumber = false;
-	constexpr bool bIncludeCameraName = false;
-	UE::MoviePipeline::ValidateOutputFormatString(FileNameFormatString, bIncludeRenderPass, bTestFrameNumber, bIncludeCameraName);
-	FString FinalFilePath = UMovieGraphBlueprintLibrary::ResolveFilenameFormatArguments(FileNameFormatString, Params, FinalFormatArgs);
-
-	if (FPaths::IsRelative(FinalFilePath))
-	{
-		FinalFilePath = FPaths::ConvertRelativePathToFull(FinalFilePath);
-	}
-
-	// If the end user opts to delete existing files and one with the same name exists, delete it
-	if (bOverwriteExistingOutput && !UE::MoviePipeline::CanWriteToFile(*FinalFilePath, false)) // false otherwise it will always return true
-	{
-		FPlatformFileManager::Get().GetPlatformFile().DeleteFile(*FinalFilePath);
-	}
+	const FString FinalFilePath = UE::MovieGraph::Capture::PrepareCapturePath(this, EvaluatedConfig, DebugSetting->UnrealInsightsTraceFileNameFormat, TEXT("utrace"), false);
 
 	const bool bTraceStarted = FTraceAuxiliary::Start(FTraceAuxiliary::EConnectionType::File, *FinalFilePath);
+
 	if (bTraceStarted)
 	{
 		bDidStartInsightsCapture = true;
diff --git a/Engine/Plugins/MovieScene/MovieRenderPipeline/Source/MovieRenderPipelineCore/Private/Graph/MovieGraphUtils.cpp b/Engine/Plugins/MovieScene/MovieRenderPipeline/Source/MovieRenderPipelineCore/Private/Graph/MovieGraphUtils.cpp
--- a/Engine/Plugins/MovieScene/MovieRenderPipeline/Source/MovieRenderPipelineCore/Private/Graph/MovieGraphUtils.cpp
+++ b/Engine/Plugins/MovieScene/MovieRenderPipeline/Source/MovieRenderPipelineCore/Private/Graph/MovieGraphUtils.cpp
@@ -2,17 +2,20 @@
 
 #include "MovieGraphUtils.h"
 
-#include "Graph/MovieGraphConfig.h"
-#include "MovieRenderPipelineCoreModule.h"
-
 #include "AudioMixerDevice.h"
 #include "AudioThread.h"
 #include "Engine/Engine.h"
+#include "Graph/MovieGraphConfig.h"
+#include "Graph/MovieGraphPipeline.h"
+#include "Graph/MovieGraphBlueprintLibrary.h"
+#include "Graph/Nodes/MovieGraphGlobalOutputSettingNode.h"
+#include "HAL/PlatformFileManager.h"
+#include "MoviePipelineUtils.h"
+#include "MovieRenderPipelineCoreModule.h"
 
 #if WITH_EDITOR
 #include "Engine/RendererSettings.h"
 #include "Framework/Notifications/NotificationManager.h"
-#include "HAL/PlatformFileManager.h"
 #include "ISettingsEditorModule.h"
 #include "Misc/ConfigCacheIni.h"
 #include "Modules/ModuleManager.h"
@@ -139,6 +142,56 @@ namespace UE::MovieGraph
 		}
 	}
 
+	namespace Capture
+	{
+		FString PrepareCapturePath(
+			const UMovieGraphPipeline* InPipeline,
+			const TObjectPtr<UMovieGraphEvaluatedConfig>& InEvaluatedConfig,
+			const FString& InPathFormat, const FString& InExt, bool InForShot)
+		{
+			check(InEvaluatedConfig);
+			check(InPipeline);
+
+			constexpr bool bIncludeCDOs = true;
+			constexpr bool bExactMatch = true;
+			const UMovieGraphGlobalOutputSettingNode* OutputSetting =
+				InEvaluatedConfig->GetSettingForBranch<UMovieGraphGlobalOutputSettingNode>(UMovieGraphSettingNode::GlobalsPinName, bIncludeCDOs, bExactMatch);
+			FString FileNameFormatString = OutputSetting->OutputDirectory.Path / InPathFormat;
+
+			const bool bOverwriteExistingOutput = OutputSetting->bOverwriteExistingOutput;
+
+			// Generate a filename for this encoded file
+			TMap<FString, FString> FormatOverrides;
+			FormatOverrides.Add(TEXT("ext"), InExt);
+
+			FMovieGraphRenderDataIdentifier TempRenderDataIdentifier;
+			TempRenderDataIdentifier.RootBranchName = UMovieGraphSettingNode::GlobalsPinName;
+
+			FMovieGraphFilenameResolveParams Params = FMovieGraphFilenameResolveParams::MakeResolveParams(
+				TempRenderDataIdentifier, InPipeline, InEvaluatedConfig, InPipeline->GetCurrentTraversalContext(InForShot), FormatOverrides);
+
+			FMovieGraphResolveArgs FinalFormatArgs;
+			constexpr bool bIncludeRenderPass = false;
+			const bool bTestFrameNumber = InForShot;
+			constexpr bool bIncludeCameraName = false;
+			UE::MoviePipeline::ValidateOutputFormatString(FileNameFormatString, bIncludeRenderPass, bTestFrameNumber, bIncludeCameraName);
+			FString FinalPath = UMovieGraphBlueprintLibrary::ResolveFilenameFormatArguments(FileNameFormatString, Params, FinalFormatArgs);
+
+			if (FPaths::IsRelative(FinalPath))
+			{
+				FinalPath = FPaths::ConvertRelativePathToFull(FinalPath);
+			}
+
+			// If the end user opts to delete existing path and one with the same name exists, delete it
+			if (bOverwriteExistingOutput && !UE::MoviePipeline::CanWriteToFile(*FinalPath, false)) // false otherwise it will always return true
+			{
+				FPlatformFileManager::Get().GetPlatformFile().DeleteDirectoryRecursively(*FinalPath);
+			}
+
+			return FinalPath;
+		}
+	}
+
 #if WITH_EDITOR
 	void UpdateDependentPropertyInConfigFile(URendererSettings* RendererSettings, FProperty* RendererProperty)
 	{
diff --git a/Engine/Plugins/MovieScene/MovieRenderPipeline/Source/MovieRenderPipelineCore/Private/Graph/MovieGraphUtils.h b/Engine/Plugins/MovieScene/MovieRenderPipeline/Source/MovieRenderPipelineCore/Private/Graph/MovieGraphUtils.h
--- a/Engine/Plugins/MovieScene/MovieRenderPipeline/Source/MovieRenderPipelineCore/Private/Graph/MovieGraphUtils.h
+++ b/Engine/Plugins/MovieScene/MovieRenderPipeline/Source/MovieRenderPipelineCore/Private/Graph/MovieGraphUtils.h
@@ -7,6 +7,8 @@
 // Forward declares
 class FAudioDevice;
 class UMovieGraphPin;
+class UMovieGraphPipeline;
+class UMovieGraphEvaluatedConfig;
 struct FMovieGraphEvaluationContext;
 namespace Audio
 {
@@ -45,6 +47,14 @@ namespace UE::MovieGraph
 		bool IsMoviePipelineAudioOutputSupported(const UObject* InWorldContextObject);
 	}
 
+	namespace Capture
+	{
+		FString PrepareCapturePath(
+			const UMovieGraphPipeline* InPipeline,
+			const TObjectPtr<UMovieGraphEvaluatedConfig>& InEvaluatedConfig,
+			const FString& InPathFormat, const FString& InExt, bool InForShot);
+	}
+
 #if WITH_EDITOR
 	void ValidateAlphaProjectSettings(const FText& InRequestingFeatureLabel, bool bMandatePrimitiveAlphaHoldout = false);
 #endif
diff --git a/Engine/Plugins/MovieScene/MovieRenderPipeline/Source/MovieRenderPipelineCore/Public/Graph/Nodes/MovieGraphDebugNode.h b/Engine/Plugins/MovieScene/MovieRenderPipeline/Source/MovieRenderPipelineCore/Public/Graph/Nodes/MovieGraphDebugNode.h
--- a/Engine/Plugins/MovieScene/MovieRenderPipeline/Source/MovieRenderPipelineCore/Public/Graph/Nodes/MovieGraphDebugNode.h
+++ b/Engine/Plugins/MovieScene/MovieRenderPipeline/Source/MovieRenderPipelineCore/Public/Graph/Nodes/MovieGraphDebugNode.h
@@ -38,32 +38,68 @@ public:
 	UPROPERTY(EditAnywhere, BlueprintReadWrite, Category = Overrides, meta = (InlineEditConditionToggle))
 	uint8 bOverride_UnrealInsightsTraceFileNameFormat : 1;
 
+	UPROPERTY(EditAnywhere, BlueprintReadWrite, Category = Overrides, meta = (InlineEditConditionToggle))
+	uint8 bOverride_bCaptureFramesWithDumpGpu : 1;
+
+	UPROPERTY(EditAnywhere, BlueprintReadWrite, Category = Overrides, meta = (InlineEditConditionToggle))
+	uint8 bOverride_DumpGpuCaptureFrameStart : 1;
+
+	UPROPERTY(EditAnywhere, BlueprintReadWrite, Category = Overrides, meta = (InlineEditConditionToggle))
+	uint8 bOverride_DumpGpuCaptureFrameStop : 1;
+
+	UPROPERTY(EditAnywhere, BlueprintReadWrite, Category = Overrides, meta = (InlineEditConditionToggle))
+	uint8 bOverride_DumpGpuCaptureFilePathFormat : 1;
+
 	/** 
 	* If true, automatically trigger RenderDoc to capture rendering information. RenderDoc plugin must be enabled, 
 	* and the editor must have been launched with -AttachRenderDoc. Resulting capture will be in /Saved/RenderDocCaptures.
 	*/
-	UPROPERTY(EditAnywhere, BlueprintReadWrite, Category = "Settings", meta = (EditCondition = "bOverride_bCaptureFramesWithRenderDoc"))
+	UPROPERTY(EditAnywhere, BlueprintReadWrite, Category = "RenderDoc", meta = (EditCondition = "bOverride_bCaptureFramesWithRenderDoc"))
 	bool bCaptureFramesWithRenderDoc;
 
 	/**
 	* If bCaptureFramesWithRenderDoc is true, which frame (on the root Sequencer time line) should we capture data for?
 	*/
-	UPROPERTY(EditAnywhere, BlueprintReadWrite, Category = "Settings", meta = (EditCondition = "bOverride_RenderDocCaptureFrame"))
+	UPROPERTY(EditAnywhere, BlueprintReadWrite, Category = "RenderDoc", meta = (EditCondition = "bOverride_RenderDocCaptureFrame"))
 	int32 RenderDocCaptureFrame;
 	
 	/** 
 	* If true, automatically capture an Unreal Insights trace file for the duration of the render.
 	* Resulting capture will be in the global Output Directory for the job.
 	*/
-	UPROPERTY(EditAnywhere, BlueprintReadWrite, Category = "Settings", meta = (EditCondition = "bOverride_bCaptureUnrealInsightsTrace"))
+	UPROPERTY(EditAnywhere, BlueprintReadWrite, Category = "Unreal Insights", meta = (EditCondition = "bOverride_bCaptureUnrealInsightsTrace"))
 	bool bCaptureUnrealInsightsTrace;
 
 	/** 
 	* If true, automatically capture an Unreal Insights trace file for the duration of the render.
 	* Resulting capture will be in the global Output Directory for the job.
 	*/
-	UPROPERTY(EditAnywhere, BlueprintReadWrite, Category = "Settings", meta = (EditCondition = "bOverride_UnrealInsightsTraceFileNameFormat"))
+	UPROPERTY(EditAnywhere, BlueprintReadWrite, Category = "Unreal Insights", meta = (EditCondition = "bOverride_UnrealInsightsTraceFileNameFormat"))
 	FString UnrealInsightsTraceFileNameFormat = TEXT("{sequence_name}_UnrealInsights");
+
+	/** 
+	* If true, automatically trigger DumpGPU to capture rendering information.
+	*/
+	UPROPERTY(EditAnywhere, BlueprintReadWrite, Category = "DumpGPU", meta = (EditCondition = "bOverride_bCaptureFramesWithDumpGpu"))
+	bool bCaptureFramesWithDumpGpu;
+
+	/**
+	* If bCaptureFramesWithDumpGpu is true, from which included frame (on the root Sequencer time line) start index should we capture? No start bound if negative.
+	*/
+	UPROPERTY(EditAnywhere, BlueprintReadWrite, Category = "DumpGPU", meta = (EditCondition = "bOverride_DumpGpuCaptureFrameStart"))
+	int32 DumpGpuCaptureFrameStart = -1;
+
+	/**
+	* If bCaptureFramesWithDumpGpu is true, up to which excluded frame (on the root Sequencer time line) stop index should we capture? No stop bound if negative.
+	*/
+	UPROPERTY(EditAnywhere, BlueprintReadWrite, Category = "DumpGPU", meta = (EditCondition = "bOverride_DumpGpuCaptureFrameStop"))
+	int32 DumpGpuCaptureFrameStop = -1;
+
+	/**
+	* If bCaptureFramesWithDumpGpu is true, name of the output directory for the capture (related to frame output directory).
+	*/
+	UPROPERTY(EditAnywhere, BlueprintReadWrite, Category = "DumpGPU", meta = (EditCondition = "bOverride_DumpGpuCaptureFilePathFormat"))
+	FString DumpGpuCaptureFilePathFormat = TEXT("{sequence_name}.{frame_number}_DumpGPU");
 };
 
 #undef UE_API
diff --git a/Engine/Source/Developer/FunctionalTesting/Private/ScreenshotFunctionalTestBase.cpp b/Engine/Source/Developer/FunctionalTesting/Private/ScreenshotFunctionalTestBase.cpp
--- a/Engine/Source/Developer/FunctionalTesting/Private/ScreenshotFunctionalTestBase.cpp
+++ b/Engine/Source/Developer/FunctionalTesting/Private/ScreenshotFunctionalTestBase.cpp
@@ -264,7 +264,7 @@ void AScreenshotFunctionalTestBase::RequestScreenshot()
 #if WITH_AUTOMATION_TESTS && WITH_DUMPGPU
 	if (GDumpGPUDumpOnScreenshotTest.GetValueOnGameThread() != 0)
 	{
-		FRDGBuilder::BeginResourceDump(TEXT(""));
+		FRDGBuilder::BeginResourceDump(TEXT(""), TEXT("ScreenshotTest"));
 	}
 #endif
 }
diff --git a/Engine/Source/Runtime/Engine/Private/UnrealEngine.cpp b/Engine/Source/Runtime/Engine/Private/UnrealEngine.cpp
--- a/Engine/Source/Runtime/Engine/Private/UnrealEngine.cpp
+++ b/Engine/Source/Runtime/Engine/Private/UnrealEngine.cpp
@@ -6362,7 +6362,7 @@ bool UEngine::HandleDumpGPUCommand(const TCHAR* Cmd, FOutputDevice& Ar)
 	check(IsInGameThread());
 
 	TArray<FString> Args;
-	FString ResourceDumpDirectory = FRDGBuilder::BeginResourceDump(Cmd);
+	FString ResourceDumpDirectory = FRDGBuilder::BeginResourceDump(Cmd, TEXT("DumpGPUCommand"));
 
 	if (!ResourceDumpDirectory.IsEmpty())
 	{
@@ -15402,7 +15402,7 @@ EBrowseReturnVal::Type UEngine::Browse( FWorldContext& WorldContext, FURL URL, F
 
 		if( !WorldContext.PendingNetGame )
 		{
-			// If the inital packet sent in InitNetDriver results in a socket error, HandleDisconnect() and CancelPending() may be called, which will null the PendingNetGame.
+			// If the initial packet sent in InitNetDriver results in a socket error, HandleDisconnect() and CancelPending() may be called, which will null the PendingNetGame.
 			Error = NSLOCTEXT("Engine", "PendingNetGameInitFailure", "Error initializing the network driver.").ToString();
 			BroadcastTravelFailure(WorldContext.World(), ETravelFailure::PendingNetGameCreateFailure, Error);
 			return EBrowseReturnVal::Failure;
@@ -16957,8 +16957,8 @@ void UEngine::CheckAndHandleStaleWorldObjectReferences(FWorldContext* WorldConte
 	for( TObjectIterator<UWorld> It; It; ++It )
 	{
 		UWorld* World = *It;
-		const bool bIsPersistantWorldType = (World->WorldType == EWorldType::Inactive) || (World->WorldType == EWorldType::EditorPreview);
-		if (!bIsPersistantWorldType && !WorldHasValidContext(World))
+		const bool bIsPersistentWorldType = (World->WorldType == EWorldType::Inactive) || (World->WorldType == EWorldType::EditorPreview);
+		if (!bIsPersistentWorldType && !WorldHasValidContext(World))
 		{
 			if ((World->PersistentLevel == nullptr || !WorldHasValidContext(World->PersistentLevel->OwningWorld)) && !IsWorldDuplicate(World))
 			{
diff --git a/Engine/Source/Runtime/RenderCore/Private/DumpGPU.cpp b/Engine/Source/Runtime/RenderCore/Private/DumpGPU.cpp
--- a/Engine/Source/Runtime/RenderCore/Private/DumpGPU.cpp
+++ b/Engine/Source/Runtime/RenderCore/Private/DumpGPU.cpp
@@ -246,6 +246,7 @@ public:
 	int32 FrameCount = 1;
 	FName UploadResourceCompressionName;
 	FString DumpingDirectoryPath;
+	FString Context;
 	FDateTime Time;
 	FGenericPlatformMemoryConstants MemoryConstants;
 	FGenericPlatformMemoryStats MemoryStats;
@@ -1960,10 +1961,18 @@ static int32 GNextDumpingRemaingFrames = -1;
 static FRDGResourceDumpContext* GNextRDGResourceDumpContext = nullptr;
 
 
-FString FRDGBuilder::BeginResourceDump(const TCHAR* Cmd)
+FString FRDGBuilder::BeginResourceDump(const TCHAR* Cmd, const TCHAR* Context)
 {
 	check(IsInGameThread());
 
+	if (Cmd == nullptr)
+	{
+		Cmd = TEXT("");
+	}
+	if (Context == nullptr)
+	{
+		Context = TEXT("");
+	}
 	if (GRDGResourceDumpContext_GameThread || GNextRDGResourceDumpContext)
 	{
 		return FString();
@@ -2027,10 +2036,18 @@ FString FRDGBuilder::BeginResourceDump(const TCHAR* Cmd)
 		}
 		NewResourceDumpContext->DumpingDirectoryPath = DirectoryPath / FApp::GetProjectName() + TEXT("-") + FPlatformProperties::PlatformName() + TEXT("-") + NewResourceDumpContext->Time.ToString() + TEXT("/");
 	}
+
+	// Override dump directory CVar
+	if (Parameters.Contains(TEXT("path")))
+	{
+		NewResourceDumpContext->DumpingDirectoryPath = Parameters[TEXT("path")];
+	}
+
 	NewResourceDumpContext->bStream = GDumpGPUStream.GetValueOnGameThread() != 0;
 	NewResourceDumpContext->bEnableDiskWrite = GDumpTestEnableDiskWrite.GetValueOnGameThread() != 0;
 	NewResourceDumpContext->DeltaTime = FApp::GetDeltaTime();
 	NewResourceDumpContext->FrameCount = FMath::Max(GDumpGPUFrameCount.GetValueOnGameThread(), 1);
+	NewResourceDumpContext->Context = Context;
 
 	if (Switches.Contains(TEXT("upload")))
 	{
@@ -2068,6 +2085,12 @@ FString FRDGBuilder::BeginResourceDump(const TCHAR* Cmd)
 
 	NewResourceDumpContext->bShowInExplore = NewResourceDumpContext->bEnableDiskWrite && GDumpExploreCVar.GetValueOnGameThread() != 0 && !NewResourceDumpContext->bUpload;
 
+	// Override open explorer behavior
+	if (Switches.Contains(TEXT("noexplorer")))
+	{
+		NewResourceDumpContext->bShowInExplore = false;
+	}
+
 	check(GNextRDGResourceDumpContext == nullptr);
 	GNextRDGResourceDumpContext = NewResourceDumpContext;
 	if (GDumpGPUDelay.GetValueOnGameThread() > 0.0f)
@@ -2190,6 +2213,7 @@ void FRDGResourceDumpContext::Start()
 		JsonObject->SetStringField(TEXT("RHI"), GDynamicRHI->GetName());
 		JsonObject->SetStringField(TEXT("RHIMaxFeatureLevel"), LexToString(GMaxRHIFeatureLevel));
 		JsonObject->SetStringField(TEXT("DumpTime"), Time.ToString());
+		JsonObject->SetStringField(TEXT("DumpContext"), Context);
 		{
 			const FString LogSrcAbsolute = FPlatformOutputDevices::GetAbsoluteLogFilename();
 			FString LogFilename = FPaths::GetCleanFilename(LogSrcAbsolute);
@@ -2659,13 +2683,20 @@ bool FRDGBuilder::IsDumpingFrame()
 namespace UE::RenderCore::DumpGPU
 {
 
-RENDERCORE_API FDumpScope::FDumpScope(bool bInCondition)
+RENDERCORE_API FDumpScope::FDumpScope(bool bInCondition, const TCHAR* Cmd, const TCHAR* Context)
 	: bCondition(bInCondition && !IsDumpingFrame())
 {
 	if (bCondition)
 	{
 		// Pass "-oneframe" to override CVar that could enable multiple frames of capture
-		FRDGBuilder::BeginResourceDump(TEXT("-oneframe"));
+		if (Cmd == nullptr)
+		{
+			Cmd = TEXT("");
+		}
+
+		const FString CmdWithOneFrame = FString::Printf(TEXT("-oneframe %s"), Cmd);
+
+		FRDGBuilder::BeginResourceDump(*CmdWithOneFrame, Context);
 
 		// Tick the DumpGPU system, which will start the dump
 		TickEndFrame();
diff --git a/Engine/Source/Runtime/RenderCore/Public/DumpGPU.h b/Engine/Source/Runtime/RenderCore/Public/DumpGPU.h
--- a/Engine/Source/Runtime/RenderCore/Public/DumpGPU.h
+++ b/Engine/Source/Runtime/RenderCore/Public/DumpGPU.h
@@ -45,7 +45,7 @@ RENDERCORE_API bool ShouldCameraCut();
 class FDumpScope
 {
 public:
-	RENDERCORE_API FDumpScope(bool bInCondition = true);
+	RENDERCORE_API FDumpScope(bool bInCondition = true, const TCHAR* Cmd = TEXT(""), const TCHAR* Context = TEXT(""));
 
 	~FDumpScope()
 	{
diff --git a/Engine/Source/Runtime/RenderCore/Public/RenderGraphBuilder.h b/Engine/Source/Runtime/RenderCore/Public/RenderGraphBuilder.h
--- a/Engine/Source/Runtime/RenderCore/Public/RenderGraphBuilder.h
+++ b/Engine/Source/Runtime/RenderCore/Public/RenderGraphBuilder.h
@@ -434,7 +434,7 @@ public:
 	FRDGBlackboard Blackboard;
 
 #if RDG_DUMP_RESOURCES
-	static RENDERCORE_API FString BeginResourceDump(const TCHAR* Cmd);
+	static RENDERCORE_API FString BeginResourceDump(const TCHAR* Cmd = TEXT(""), const TCHAR* Context = TEXT(""));
 	static RENDERCORE_API bool IsDumpingFrame();
 #else
 	static bool IsDumpingFrame() { return false; }
