[MRG] Cleanup after DumpGPU integration

Problem: Bughawk discovered potential trouble following DumpGPU integration
Solution: Address relevant feedback.

#rnx
#tests win64 editor : generate DumpGPU and UnrealInsight traces via MRG (with and without previous file/folder)

diff --git a/Engine/Plugins/MovieScene/MovieRenderPipeline/Source/MovieRenderPipelineCore/Private/Graph/MovieGraphDefaultRenderer.cpp b/Engine/Plugins/MovieScene/MovieRenderPipeline/Source/MovieRenderPipelineCore/Private/Graph/MovieGraphDefaultRenderer.cpp
--- a/Engine/Plugins/MovieScene/MovieRenderPipeline/Source/MovieRenderPipelineCore/Private/Graph/MovieGraphDefaultRenderer.cpp
+++ b/Engine/Plugins/MovieScene/MovieRenderPipeline/Source/MovieRenderPipelineCore/Private/Graph/MovieGraphDefaultRenderer.cpp
@@ -341,14 +341,13 @@ void UMovieGraphDefaultRenderer::Render(const FMovieGraphTimeStepData& InTimeSte
 	{
 	
 		const UMovieGraphDebugSettingNode* DebugSettings = InTimeStepData.EvaluatedConfig->GetSettingForBranch<UMovieGraphDebugSettingNode>(UMovieGraphNode::GlobalsPinName);
+		check(DebugSettings);
 		const bool bIsRendering = GetOwningGraph()->GetActiveShotList()[GetOwningGraph()->GetCurrentShotIndex()]->ShotInfo.State == EMovieRenderShotState::Rendering;
-
-		const bool bWantsRenderDocDebugCapture = DebugSettings && DebugSettings->bCaptureFramesWithRenderDoc;
-		const bool bWantsDumpGpuDebugCapture = DebugSettings && DebugSettings->bCaptureFramesWithDumpGpu;
-
+		const bool bWantsRenderDocDebugCapture = DebugSettings->bCaptureFramesWithRenderDoc;
+		const bool bWantsDumpGpuDebugCapture = DebugSettings->bCaptureFramesWithDumpGpu;
 		const bool bIsRenderDocCorrectFrame = InTimeStepData.RootFrameNumber == DebugSettings->RenderDocCaptureFrame;
-		const bool bIsDumpGPUCorrectFrame = ((DebugSettings->DumpGpuCaptureFrameStart < 0 || InTimeStepData.RootFrameNumber >= DebugSettings->DumpGpuCaptureFrameStart) &&
-									         (DebugSettings->DumpGpuCaptureFrameStop  < 0 || InTimeStepData.RootFrameNumber <  DebugSettings->DumpGpuCaptureFrameStop));
+		const bool bIsDumpGPUCorrectFrame = (DebugSettings->DumpGpuCaptureFrameStart < 0 || InTimeStepData.RootFrameNumber >= DebugSettings->DumpGpuCaptureFrameStart) &&
+											(DebugSettings->DumpGpuCaptureFrameStop  < 0 || InTimeStepData.RootFrameNumber <  DebugSettings->DumpGpuCaptureFrameStop);
 
 		if (bIsRendering && bWantsRenderDocDebugCapture && bIsRenderDocCorrectFrame)
 		{
diff --git a/Engine/Plugins/MovieScene/MovieRenderPipeline/Source/MovieRenderPipelineCore/Private/Graph/MovieGraphPipeline.cpp b/Engine/Plugins/MovieScene/MovieRenderPipeline/Source/MovieRenderPipelineCore/Private/Graph/MovieGraphPipeline.cpp
--- a/Engine/Plugins/MovieScene/MovieRenderPipeline/Source/MovieRenderPipelineCore/Private/Graph/MovieGraphPipeline.cpp
+++ b/Engine/Plugins/MovieScene/MovieRenderPipeline/Source/MovieRenderPipelineCore/Private/Graph/MovieGraphPipeline.cpp
@@ -793,7 +793,7 @@ void UMovieGraphPipeline::StartUnrealInsightsCapture(UMovieGraphEvaluatedConfig*
 	const UMovieGraphDebugSettingNode* DebugSetting =
 		EvaluatedConfig->GetSettingForBranch<UMovieGraphDebugSettingNode>(UMovieGraphSettingNode::GlobalsPinName, bIncludeCDOs, bExactMatch);
 
-	if (!ensureAlwaysMsgf(DebugSetting, TEXT("Failed to find UMovieGraphFileOutputNode. Aborting.")))
+	if (!ensureAlwaysMsgf(DebugSetting, TEXT("Failed to find UMovieGraphDebugSettingNode. Aborting.")))
 	{
 		return;
 	}
diff --git a/Engine/Plugins/MovieScene/MovieRenderPipeline/Source/MovieRenderPipelineCore/Private/Graph/MovieGraphUtils.cpp b/Engine/Plugins/MovieScene/MovieRenderPipeline/Source/MovieRenderPipelineCore/Private/Graph/MovieGraphUtils.cpp
--- a/Engine/Plugins/MovieScene/MovieRenderPipeline/Source/MovieRenderPipelineCore/Private/Graph/MovieGraphUtils.cpp
+++ b/Engine/Plugins/MovieScene/MovieRenderPipeline/Source/MovieRenderPipelineCore/Private/Graph/MovieGraphUtils.cpp
@@ -182,10 +182,16 @@ namespace UE::MovieGraph
 				FinalPath = FPaths::ConvertRelativePathToFull(FinalPath);
 			}
 
-			// If the end user opts to delete existing path and one with the same name exists, delete it
-			if (bOverwriteExistingOutput && !UE::MoviePipeline::CanWriteToFile(*FinalPath, false)) // false otherwise it will always return true
+			if (bOverwriteExistingOutput)
 			{
-				FPlatformFileManager::Get().GetPlatformFile().DeleteDirectoryRecursively(*FinalPath);
+				if (FPlatformFileManager::Get().GetPlatformFile().DirectoryExists(*FinalPath))
+				{
+					FPlatformFileManager::Get().GetPlatformFile().DeleteDirectoryRecursively(*FinalPath);
+				} 
+				else if (!UE::MoviePipeline::CanWriteToFile(*FinalPath, false)) // false otherwise it will always return true
+				{
+					FPlatformFileManager::Get().GetPlatformFile().DeleteFile(*FinalPath);
+				}
 			}
 
 			return FinalPath;
diff --git a/Engine/Plugins/MovieScene/MovieRenderPipeline/Source/MovieRenderPipelineCore/Public/Graph/Nodes/MovieGraphDebugNode.h b/Engine/Plugins/MovieScene/MovieRenderPipeline/Source/MovieRenderPipelineCore/Public/Graph/Nodes/MovieGraphDebugNode.h
--- a/Engine/Plugins/MovieScene/MovieRenderPipeline/Source/MovieRenderPipelineCore/Public/Graph/Nodes/MovieGraphDebugNode.h
+++ b/Engine/Plugins/MovieScene/MovieRenderPipeline/Source/MovieRenderPipelineCore/Public/Graph/Nodes/MovieGraphDebugNode.h
@@ -68,13 +68,12 @@ public:
 	
 	/** 
 	* If true, automatically capture an Unreal Insights trace file for the duration of the render.
-	* Resulting capture will be in the global Output Directory for the job.
 	*/
 	UPROPERTY(EditAnywhere, BlueprintReadWrite, Category = "Unreal Insights", meta = (EditCondition = "bOverride_bCaptureUnrealInsightsTrace"))
 	bool bCaptureUnrealInsightsTrace;
 
 	/** 
-	* If true, automatically capture an Unreal Insights trace file for the duration of the render.
+	* If bCaptureUnrealInsightsTrace is true, name of the UnrealInsights trace.
 	* Resulting capture will be in the global Output Directory for the job.
 	*/
 	UPROPERTY(EditAnywhere, BlueprintReadWrite, Category = "Unreal Insights", meta = (EditCondition = "bOverride_UnrealInsightsTraceFileNameFormat"))
