MoviePipeline: Fixed outdated tooltips regarding overscan in the graph and legacy presets.

#rnx

diff --git a/Engine/Plugins/MovieScene/MovieRenderPipeline/Source/MovieRenderPipelineCore/Public/Graph/Nodes/MovieGraphCameraNode.h b/Engine/Plugins/MovieScene/MovieRenderPipeline/Source/MovieRenderPipelineCore/Public/Graph/Nodes/MovieGraphCameraNode.h
--- a/Engine/Plugins/MovieScene/MovieRenderPipeline/Source/MovieRenderPipelineCore/Public/Graph/Nodes/MovieGraphCameraNode.h
+++ b/Engine/Plugins/MovieScene/MovieRenderPipeline/Source/MovieRenderPipelineCore/Public/Graph/Nodes/MovieGraphCameraNode.h
@@ -53,15 +53,18 @@ public:
 	EMoviePipelineShutterTiming ShutterTiming;
 
 	/**
-	* Overscan percent enables rendering render additional pixels beyond the set resolution and can be used in conjunction 
-	* with EXR file output to add post-processing effects such as lens distortion.
-	* Please note that using this feature might affect the results due to auto-exposure and other camera settings.
-	* On EXR this will produce a 1080p image with extra pixel data hidden around the outside edges for use 
-	* in post production. For all other formats this will increase the final resolution and no pixels will be hidden 
-	* (ie: 1080p with 10.0 overscan will make a 2112x1188 jpg, but a 1080p exr /w 96/54 pixels hidden on each side).
+	* Overscan percent enables rendering additional pixels beyond the current target resolution, and can be used in conjunction 
+	* with the EXR output format to add post-processing effects such as lens distortion. Please note that using this feature may cause
+	* the render to appear different in some scenarios (for example, auto-exposure may be impacted due to the increased FoV picking
+	* up additional lighting).
+	* 
+	* With EXR outputs, the non-overscanned pixel data is contained in the EXR "display window", and the overscanned pixel data is
+	* contained in the "data window". For example, a 1920x1080 render with 10% overscan will result in a 2112x1188 render; the display
+	* window will be 1920x1080, and the data window will be 2112x1188. For all other formats, the pixels are cropped to the
+	* original target resolution (ie, no extra pixel data is included).
 	*
-	* Note: This uses 0-100 and not 0-1 like the previous system did to bring it in-line with other usages
-	* of overscan in the engine (nDisplay).
+	* Note: This uses values in the [0-100] range, and not [0-1] like the preset system, to bring it in-line with other usages
+	* of overscan in the engine (like nDisplay).
 	*/
 	UPROPERTY(EditAnywhere, BlueprintReadWrite, meta = (DisplayName = "Overscan Percentage Override", UIMin = "0", UIMax = "100", ClampMin = "0", ClampMax = "100"), Category = "Settings", meta = (EditCondition = "bOverride_OverscanPercentage"))
 	float OverscanPercentage;
diff --git a/Engine/Plugins/MovieScene/MovieRenderPipeline/Source/MovieRenderPipelineCore/Public/MoviePipelineCameraSetting.h b/Engine/Plugins/MovieScene/MovieRenderPipeline/Source/MovieRenderPipelineCore/Public/MoviePipelineCameraSetting.h
--- a/Engine/Plugins/MovieScene/MovieRenderPipeline/Source/MovieRenderPipelineCore/Public/MoviePipelineCameraSetting.h
+++ b/Engine/Plugins/MovieScene/MovieRenderPipeline/Source/MovieRenderPipelineCore/Public/MoviePipelineCameraSetting.h
@@ -47,12 +47,15 @@ public:
 	bool bOverrideCameraOverscan;
 	
 	/**
-	* Overscan percent allows to render additional pixels beyond the set resolution and can be used in conjunction 
-	* with EXR file output to add post-processing effects such as lens distortion.
-	* Please note that using this feature might affect the results due to auto-exposure and other camera settings.
-	* On EXR this will produce a 1080p image with extra pixel data hidden around the outside edges for use 
-	* in post production. For all other formats this will increase the final resolution and no pixels will be hidden 
-	* (ie: 1080p /w 0.1 overscan will make a 2112x1188 jpg, but a 1080p exr /w 96/54 pixels hidden on each side)
+	* Overscan percent enables rendering additional pixels beyond the current target resolution, and can be used in conjunction 
+	* with the EXR output format to add post-processing effects such as lens distortion. Please note that using this feature may cause
+	* the render to appear different in some scenarios (for example, auto-exposure may be impacted due to the increased FoV picking
+	* up additional lighting).
+	* 
+	* With EXR outputs, the non-overscanned pixel data is contained in the EXR "display window", and the overscanned pixel data is
+	* contained in the "data window". For example, a 1920x1080 render with 0.1 overscan will result in a 2112x1188 render; the display
+	* window will be 1920x1080, and the data window will be 2112x1188. For all other formats, the pixels are cropped to the
+	* original target resolution (ie, no extra pixel data is included).
 	*/
 	UPROPERTY(EditAnywhere, BlueprintReadWrite, meta = (DisplayName = "Overscan Percentage Override", UIMin = "0", UIMax = "1", ClampMin = "0", ClampMax = "1", EditCondition = "bOverrideCameraOverscan"), Category = "Camera Settings")
 	float OverscanPercentage;
