[MRG+DumpGPU] MRG Allow filename without extension + improve error reporting of DumpGPU

Problems: MRG always add a . filename not allowing file without extension + DumpGPU is missing a log when it cannot dump a resource to disk
Solution: allows extension less files + add relevant log to DumpGPU

#rnx
#tests win64 editor : generate DumpGPU with MRG

diff --git a/Engine/Extras/GPUDumpViewer/GPUDumpViewer.html b/Engine/Extras/GPUDumpViewer/GPUDumpViewer.html
--- a/Engine/Extras/GPUDumpViewer/GPUDumpViewer.html
+++ b/Engine/Extras/GPUDumpViewer/GPUDumpViewer.html
@@ -2323,7 +2323,7 @@ class TextureView extends ResourceView
 			}
 			else
 			{
-				texture_view.error(`Error: Resources/${subresource_version_name}.bin couldn't be loaded.`);
+				texture_view.error(`Error: Resources/${subresource_version_name}.bin couldn't be loaded. This might be caused by OS path limitation. Try moving the dump folder closer to root.`);
 			}
 		});
 
@@ -5138,7 +5138,7 @@ void main(void) {
 			}
 			else
 			{
-				raytracing_view.error(`Error: Resources/${subresource_version_name}.bin couldn't be loaded.`);
+				raytracing_view.error(`Error: Resources/${subresource_version_name}.bin couldn't be loaded. This might be caused by OS path limitation. Try moving the dump folder closer to root.`);
 			}
 		});
 
diff --git a/Engine/Plugins/MovieScene/MovieRenderPipeline/Source/MovieRenderPipelineCore/Private/Graph/MovieGraphBlueprintLibrary.cpp b/Engine/Plugins/MovieScene/MovieRenderPipeline/Source/MovieRenderPipelineCore/Private/Graph/MovieGraphBlueprintLibrary.cpp
--- a/Engine/Plugins/MovieScene/MovieRenderPipeline/Source/MovieRenderPipelineCore/Private/Graph/MovieGraphBlueprintLibrary.cpp
+++ b/Engine/Plugins/MovieScene/MovieRenderPipeline/Source/MovieRenderPipelineCore/Private/Graph/MovieGraphBlueprintLibrary.cpp
@@ -364,8 +364,8 @@ namespace
 		// If the extension is not resolved, ThisTry will be a path
 		const FString ExtToken = TEXT(".{ext}");
 		const FString Extension = FString::Format(*ExtToken, NamedArgs);
-		const FString ResolvedExtension = (Extension == ExtToken) ? FString() : Extension;
-		FString ThisTry = FormattedString + Extension;
+		const FString ResolvedExtension = (Extension == ExtToken || Extension == TEXT(".")) ? FString() : Extension;
+		FString ThisTry = FormattedString + ResolvedExtension;
 		
 		// Check that it's a valid path that we can write to.
 		if (UE::MoviePipeline::CanWriteToFile(*ThisTry, bOverwriteExisting))
diff --git a/Engine/Source/Runtime/RenderCore/Private/DumpGPU.cpp b/Engine/Source/Runtime/RenderCore/Private/DumpGPU.cpp
--- a/Engine/Source/Runtime/RenderCore/Private/DumpGPU.cpp
+++ b/Engine/Source/Runtime/RenderCore/Private/DumpGPU.cpp
@@ -1284,7 +1284,12 @@ public:
 
 				RHICmdList.UnmapStagingSurface(StagingTexture, GPUIndex);
 
-				DumpResourceBinaryToFile(Array.GetData(), Array.Num(), DumpFilePath);
+				bool bDumpSucceeded = DumpResourceBinaryToFile(Array.GetData(), Array.Num(), DumpFilePath);
+
+				if (!bDumpSucceeded)
+				{
+					UE_LOG(LogDumpGPU, Warning, TEXT("DumpResourceBinaryToFile() to dump texture %s failed."), TextureDebugName);
+				}
 			}
 			else
 			{
