MoviePipeline: Added support for the "Play After Render" feature into the main MRQ queue (was previously only available in Quick Render).

Executors are now responsible for controlling this feature. They have a new SetPostRenderAction() method on them, so any executor can hook into this if desired. Currently, the PIE executor is the only shipping executor that supports this. By default, executors will return "false" from GetSupportsPostRenderActions().

Also fixed a minor bug in the "DollarF" frame range notation in post-render playback (a dash was added in the frame range when it shouldn't have been).

diff --git a/Engine/Plugins/MovieScene/MovieRenderPipeline/Config/DefaultMovieRenderPipeline.ini b/Engine/Plugins/MovieScene/MovieRenderPipeline/Config/DefaultMovieRenderPipeline.ini
--- a/Engine/Plugins/MovieScene/MovieRenderPipeline/Config/DefaultMovieRenderPipeline.ini
+++ b/Engine/Plugins/MovieScene/MovieRenderPipeline/Config/DefaultMovieRenderPipeline.ini
@@ -11,4 +11,5 @@
 +ClassRedirects=(OldName="/Script/MovieRenderPipelineCore.MoviePipelineCollectionModifier",NewName="/Script/MovieRenderPipelineCore.MovieGraphCollectionModifier")
 +ClassRedirects=(OldName="/Script/MovieRenderPipelineCore.MoviePipelineMaterialModifier",NewName="/Script/MovieRenderPipelineCore.MovieGraphMaterialModifier")
 +ClassRedirects=(OldName="/Script/MovieRenderPipelineCore.MoviePipelineVisibilityModifier",NewName="/Script/MovieRenderPipelineCore.MovieGraphRenderPropertyModifier")
-+ClassRedirects=(OldName="/Script/MovieRenderPipelineCore.MoviePipelineCollection",NewName="/Script/MovieRenderPipelineCore.MovieGraphCollection")
\ No newline at end of file
++ClassRedirects=(OldName="/Script/MovieRenderPipelineCore.MoviePipelineCollection",NewName="/Script/MovieRenderPipelineCore.MovieGraphCollection")
++EnumRedirects=(OldName="/Script/MovieRenderPipelineCore.EMovieGraphQuickRenderPostRenderActionType",NewName="/Script/MovieRenderPipelineCore.EMoviePipelinePostRenderActionType")
\ No newline at end of file
diff --git a/Engine/Plugins/MovieScene/MovieRenderPipeline/Source/MovieRenderPipelineCore/Private/MoviePipelineExecutor.cpp b/Engine/Plugins/MovieScene/MovieRenderPipeline/Source/MovieRenderPipelineCore/Private/MoviePipelineExecutor.cpp
--- a/Engine/Plugins/MovieScene/MovieRenderPipeline/Source/MovieRenderPipelineCore/Private/MoviePipelineExecutor.cpp
+++ b/Engine/Plugins/MovieScene/MovieRenderPipeline/Source/MovieRenderPipelineCore/Private/MoviePipelineExecutor.cpp
@@ -200,6 +200,18 @@ void UMoviePipelineExecutorBase::OnBeginFrame_Implementation()
 	ProcessIncomingSocketData();
 }
 
+void UMoviePipelineExecutorBase::SetPostRenderAction_Implementation(const EMoviePipelinePostRenderActionType PostRenderAction)
+{
+	if (GetSupportsPostRenderActions())
+	{
+		PostRenderActionType = PostRenderAction;
+	}
+	else
+	{
+		UE_LOG(LogMovieRenderPipeline, Warning, TEXT("Attempting to set a post-render action on an executor that does not support them."));
+	}
+}
+
 int32 UMoviePipelineExecutorBase::SendHTTPRequest(const FString& InURL, const FString& InVerb, const FString& InMessage, const TMap<FString, FString>& InHeaders)
 {
 	FHttpRequestPtr HttpRequest = FHttpModule::Get().CreateRequest();
diff --git a/Engine/Plugins/MovieScene/MovieRenderPipeline/Source/MovieRenderPipelineCore/Private/MoviePipelineQueue.cpp b/Engine/Plugins/MovieScene/MovieRenderPipeline/Source/MovieRenderPipelineCore/Private/MoviePipelineQueue.cpp
--- a/Engine/Plugins/MovieScene/MovieRenderPipeline/Source/MovieRenderPipelineCore/Private/MoviePipelineQueue.cpp
+++ b/Engine/Plugins/MovieScene/MovieRenderPipeline/Source/MovieRenderPipelineCore/Private/MoviePipelineQueue.cpp
@@ -382,6 +382,19 @@ UMovieJobVariableAssignmentContainer* UMoviePipelineExecutorShot::GetOrCreateJob
 	return nullptr;
 }
 
+#if WITH_EDITOR
+void UMoviePipelineQueueProjectSettings::PostEditChangeProperty(FPropertyChangedEvent& PropertyChangedEvent)
+{
+	Super::PostEditChangeProperty(PropertyChangedEvent);
+
+	// Save the settings as soon as they're changed.
+	if (PropertyChangedEvent.ChangeType != EPropertyChangeType::Interactive)
+	{
+		SaveConfig();
+	}
+}
+#endif	// WITH_EDITOR
+
 UMoviePipelineShotConfig* UMoviePipelineExecutorShot::AllocateNewShotOverrideConfig(TSubclassOf<UMoviePipelineShotConfig> InConfigType)
 {
 	if (!ensureAlwaysMsgf(InConfigType, TEXT("Failed to specify a Config Type. Use the default in project setting or UMoviePipelineShotConfig.")))
diff --git a/Engine/Plugins/MovieScene/MovieRenderPipeline/Source/MovieRenderPipelineCore/Private/MoviePipelineTelemetry.cpp b/Engine/Plugins/MovieScene/MovieRenderPipeline/Source/MovieRenderPipelineCore/Private/MoviePipelineTelemetry.cpp
--- a/Engine/Plugins/MovieScene/MovieRenderPipeline/Source/MovieRenderPipelineCore/Private/MoviePipelineTelemetry.cpp
+++ b/Engine/Plugins/MovieScene/MovieRenderPipeline/Source/MovieRenderPipelineCore/Private/MoviePipelineTelemetry.cpp
@@ -77,7 +77,7 @@ void FMoviePipelineTelemetry::SendQuickRenderRequestedTelemetry(const EMovieGrap
 	EventAttributes.Add(FAnalyticsEventAttribute(TEXT("Mode"),
 		StaticEnum<EMovieGraphQuickRenderMode>()->GetNameStringByValue(static_cast<int64>(QuickRenderMode))));
 	EventAttributes.Add(FAnalyticsEventAttribute(TEXT("PlayAfterRenderMode"),
-		StaticEnum<EMovieGraphQuickRenderPostRenderActionType>()->GetNameStringByValue(static_cast<int64>(ModeSettings->PostRenderBehavior))));
+		StaticEnum<EMoviePipelinePostRenderActionType>()->GetNameStringByValue(static_cast<int64>(ModeSettings->PostRenderBehavior))));
 	EventAttributes.Add(FAnalyticsEventAttribute(TEXT("ViewportLookFlags"),
 		StaticEnum<EMovieGraphQuickRenderViewportLookFlags>()->GetValueOrBitfieldAsString(ModeSettings->ViewportLookFlags)));
 	EventAttributes.Add(FAnalyticsEventAttribute(TEXT("SequencerFrameRange"),
diff --git a/Engine/Plugins/MovieScene/MovieRenderPipeline/Source/MovieRenderPipelineCore/Private/MoviePipelineTiming.cpp b/Engine/Plugins/MovieScene/MovieRenderPipeline/Source/MovieRenderPipelineCore/Private/MoviePipelineTiming.cpp
--- a/Engine/Plugins/MovieScene/MovieRenderPipeline/Source/MovieRenderPipelineCore/Private/MoviePipelineTiming.cpp
+++ b/Engine/Plugins/MovieScene/MovieRenderPipeline/Source/MovieRenderPipelineCore/Private/MoviePipelineTiming.cpp
@@ -61,7 +61,12 @@ void UMoviePipeline::ProcessOutstandingFutures()
 		}
 
 		// Add the filepath to the renderpass data.
-		ShotOutputData->RenderPassData.FindOrAdd(FutureData.PassIdentifier).FilePaths.Add(FutureData.FilePath);
+		FMoviePipelineRenderPassOutputData& OutputData = ShotOutputData->RenderPassData.FindOrAdd(FutureData.PassIdentifier);
+		OutputData.FilePaths.Add(FutureData.FilePath);
+		OutputData.FrameTemplatedFilePaths.AddUnique(FutureData.FrameTemplatedFilePath);
+
+		FMoviePipelineStringArray& FilesFromSettingClass = OutputData.SettingTypeToFilePaths.FindOrAdd(FutureData.OriginSettingClass);
+		FilesFromSettingClass.Array.Add(FutureData.FilePath);
 
 		if (!OutputFuture.Get<0>().Get())
 		{
diff --git a/Engine/Plugins/MovieScene/MovieRenderPipeline/Source/MovieRenderPipelineCore/Private/MoviePipelineUtils.cpp b/Engine/Plugins/MovieScene/MovieRenderPipeline/Source/MovieRenderPipelineCore/Private/MoviePipelineUtils.cpp
--- a/Engine/Plugins/MovieScene/MovieRenderPipeline/Source/MovieRenderPipelineCore/Private/MoviePipelineUtils.cpp
+++ b/Engine/Plugins/MovieScene/MovieRenderPipeline/Source/MovieRenderPipelineCore/Private/MoviePipelineUtils.cpp
@@ -261,6 +261,31 @@ namespace MoviePipeline
 		InOutFileMetadata.Add(TEXT("unreal/shotName"), ShotName);
 	}
 
+	FString GetFrameTemplatedFilePath(const FString& InFilePath)
+	{
+		static const FString FramePlaceholder = TEXT("{frame_placeholder}");
+
+		FString FrameTemplatedFormatString = InFilePath;
+		if (FrameTemplatedFormatString.Contains(TEXT("{frame_number}")))
+		{
+			FrameTemplatedFormatString = InFilePath.Replace(TEXT("{frame_number}"), *FramePlaceholder);
+		}
+		else if (FrameTemplatedFormatString.Contains(TEXT("{frame_number_rel}")))
+		{
+			FrameTemplatedFormatString = InFilePath.Replace(TEXT("{frame_number_rel}"), *FramePlaceholder);
+		}
+		else if (FrameTemplatedFormatString.Contains(TEXT("{frame_number_shot}")))
+		{
+			FrameTemplatedFormatString = InFilePath.Replace(TEXT("{frame_number_shot}"), *FramePlaceholder);
+		}
+		else if (FrameTemplatedFormatString.Contains(TEXT("{frame_number_shot_rel}")))
+		{
+			FrameTemplatedFormatString = InFilePath.Replace(TEXT("{frame_number_shot_rel}"), *FramePlaceholder);
+		}
+
+		return FrameTemplatedFormatString;
+	}
+
 	TObjectPtr<UMovieJobVariableAssignmentContainer> GetOrCreateJobVariableAssignmentsForGraph(const UMovieGraphConfig* InGraph, TArray<TObjectPtr<UMovieJobVariableAssignmentContainer>>& InVariableAssignments, UObject* InAssignmentsOwner)
 	{
 		for (TObjectPtr<UMovieJobVariableAssignmentContainer>& VariableAssignment : InVariableAssignments)
diff --git a/Engine/Plugins/MovieScene/MovieRenderPipeline/Source/MovieRenderPipelineCore/Private/MoviePipelineVideoOutputBase.cpp b/Engine/Plugins/MovieScene/MovieRenderPipeline/Source/MovieRenderPipelineCore/Private/MoviePipelineVideoOutputBase.cpp
--- a/Engine/Plugins/MovieScene/MovieRenderPipeline/Source/MovieRenderPipelineCore/Private/MoviePipelineVideoOutputBase.cpp
+++ b/Engine/Plugins/MovieScene/MovieRenderPipeline/Source/MovieRenderPipelineCore/Private/MoviePipelineVideoOutputBase.cpp
@@ -206,6 +206,7 @@ void UMoviePipelineVideoOutputBase::OnReceiveImageDataImpl(FMoviePipelineMergerO
 				OutputData.Shot = GetPipeline()->GetActiveShotList()[Payload->SampleState.OutputState.ShotIndex];
 				OutputData.PassIdentifier = RenderPassData.Key;
 				OutputData.FilePath = FinalFilePath;
+				OutputData.OriginSettingClass = GetClass();
 
 				GetPipeline()->AddOutputFuture(Completed.GetFuture(), OutputData);
 
diff --git a/Engine/Plugins/MovieScene/MovieRenderPipeline/Source/MovieRenderPipelineCore/Public/Graph/MovieGraphQuickRenderSettings.h b/Engine/Plugins/MovieScene/MovieRenderPipeline/Source/MovieRenderPipelineCore/Public/Graph/MovieGraphQuickRenderSettings.h
--- a/Engine/Plugins/MovieScene/MovieRenderPipeline/Source/MovieRenderPipelineCore/Public/Graph/MovieGraphQuickRenderSettings.h
+++ b/Engine/Plugins/MovieScene/MovieRenderPipeline/Source/MovieRenderPipelineCore/Public/Graph/MovieGraphQuickRenderSettings.h
@@ -3,6 +3,7 @@
 #pragma once
 
 #include "CoreMinimal.h"
+#include "MovieRenderPipelineDataTypes.h"
 
 #include "MovieGraphQuickRenderSettings.generated.h"
 
@@ -56,19 +57,8 @@ enum class EMovieGraphQuickRenderFrameRangeType : uint8
 	Custom
 };
 
-/** The action that quick render should take after finishing the render. */
-UENUM(BlueprintType)
-enum class EMovieGraphQuickRenderPostRenderActionType : uint8
-{
-	/** Do nothing after the render. */
-	DoNothing,
-
-	/** Open the media that was generated from Quick Render. The application used for this is defined in Editor Preferences for Movie Render Graph. */
-	PlayRenderOutput,
-
-	/** Open the directory that contains the files generated by Quick Render. */
-	OpenOutputDirectory
-};
+UE_DEPRECATED(5.8, "Use EMoviePipelinePostRenderActionType instead")
+typedef EMoviePipelinePostRenderActionType EMovieGraphQuickRenderPostRenderActionType;
 
 /** The aspects of the viewport that should be applied to the quick render. */
 UENUM(BlueprintType, Flags, meta = (Bitflags, UseEnumValuesAsMaskValuesInEditor = "true"))
@@ -161,12 +151,12 @@ public:
 	 * property is never set directly because Sequencer is always used. However, for scripting purposes, there may be cases where Sequencer is not
 	 * open, and a specific level sequence needs to be used.
 	 */
-	UPROPERTY(BlueprintReadWrite, Category="Configuration", DisplayName="Movie Render Graph Config")
+	UPROPERTY(BlueprintReadWrite, Category="Configuration")
 	TSoftObjectPtr<class ULevelSequence> LevelSequenceOverride = nullptr;
 
 	/** The action that quick render should perform after a render is finished. */
 	UPROPERTY(EditAnywhere, BlueprintReadWrite, Category="Configuration", DisplayName="After Render")
-	EMovieGraphQuickRenderPostRenderActionType PostRenderBehavior = EMovieGraphQuickRenderPostRenderActionType::PlayRenderOutput;
+	EMoviePipelinePostRenderActionType PostRenderBehavior = EMoviePipelinePostRenderActionType::PlayRenderOutput;
 
 	/** Enables/disables the ViewportLookFlags property from taking effect. */
 	UPROPERTY(EditAnywhere, BlueprintReadWrite, Category="Quick Render", meta=(InlineEditConditionToggle))
diff --git a/Engine/Plugins/MovieScene/MovieRenderPipeline/Source/MovieRenderPipelineCore/Public/Graph/MovieRenderGraphEditorSettings.h b/Engine/Plugins/MovieScene/MovieRenderPipeline/Source/MovieRenderPipelineCore/Public/Graph/MovieRenderGraphEditorSettings.h
--- a/Engine/Plugins/MovieScene/MovieRenderPipeline/Source/MovieRenderPipelineCore/Public/Graph/MovieRenderGraphEditorSettings.h
+++ b/Engine/Plugins/MovieScene/MovieRenderPipeline/Source/MovieRenderPipelineCore/Public/Graph/MovieRenderGraphEditorSettings.h
@@ -26,8 +26,8 @@ public:
 	UE_API virtual void PostEditChangeChainProperty(FPropertyChangedChainEvent& PropertyChangedEvent) override;
 #endif
 
-	/** If PostRenderBehavior is set to PlayRenderOutput, these settings are used to determine how to play media. */
-	UPROPERTY(Config, EditAnywhere, Category="Play Render Output (Quick Render only)", meta=(ShowOnlyInnerProperties))
+	/** If playing back media post-render, these settings dictate how that will be done. */
+	UPROPERTY(Config, EditAnywhere, Category="Play Render Output", meta=(ShowOnlyInnerProperties))
 	FMovieGraphPostRenderSettings PostRenderSettings;
 };
 
diff --git a/Engine/Plugins/MovieScene/MovieRenderPipeline/Source/MovieRenderPipelineCore/Public/MoviePipelineExecutor.h b/Engine/Plugins/MovieScene/MovieRenderPipeline/Source/MovieRenderPipelineCore/Public/MoviePipelineExecutor.h
--- a/Engine/Plugins/MovieScene/MovieRenderPipeline/Source/MovieRenderPipelineCore/Public/MoviePipelineExecutor.h
+++ b/Engine/Plugins/MovieScene/MovieRenderPipeline/Source/MovieRenderPipelineCore/Public/MoviePipelineExecutor.h
@@ -49,6 +49,7 @@ class UMoviePipelineExecutorBase : public UObject
 public:
 	UMoviePipelineExecutorBase()
 		: bAnyJobHadFatalError(false)
+		, PostRenderActionType(EMoviePipelinePostRenderActionType::DoNothing)
 		, TargetPipelineClass(nullptr)
 		, ExternalSocket(nullptr)
 		, RecvMessageDataRemaining(0)
@@ -173,6 +174,42 @@ public:
 	*/
 	UFUNCTION(BlueprintCallable, BlueprintNativeEvent, Category = "Movie Render Pipeline")
 	UE_API void CancelAllJobs();
+	
+	/**
+	* Set the post-render action that will be taken when all the renders in the queue finish.
+	*
+	* For C++ implementations override `virtual void SetPostRenderAction_Implementation(const EMoviePipelinePostRenderActionType PostRenderAction) override`
+	* For Python/BP implementations override
+	*	@unreal.ufunction(override=True)
+	*	def set_post_render_action(self, postRenderAction):
+	*/
+	UFUNCTION(BlueprintCallable, BlueprintNativeEvent, Category = "Movie Render Pipeline")
+	UE_API void SetPostRenderAction(const EMoviePipelinePostRenderActionType PostRenderAction);
+
+	/**
+	* Get the post-render action that will be taken when all the renders in the queue finish.
+	*
+	* For C++ implementations override `virtual EMoviePipelinePostRenderActionType GetPostRenderAction_Implementation() override`
+	* For Python/BP implementations override
+	*	@unreal.ufunction(override=True)
+	*	def get_post_render_action(self):
+	*/
+	UFUNCTION(BlueprintCallable, BlueprintNativeEvent, Category = "Movie Render Pipeline")
+	UE_API EMoviePipelinePostRenderActionType GetPostRenderAction();
+
+	/**
+	* Get whether this executor supports performing a post-render action.
+	*
+	* The base executor classes do not implement post-render actions. Executors returning "true" here must implement support for post-render
+	* actions themselves.
+	*
+	* For C++ implementations override `virtual bool GetSupportsPostRenderActions_Implementation() override`
+	* For Python/BP implementations override
+	*	@unreal.ufunction(override=True)
+	*	def get_supports_post_render_actions(self):
+	*/
+	UFUNCTION(BlueprintCallable, BlueprintNativeEvent, Category = "Movie Render Pipeline")
+	UE_API bool GetSupportsPostRenderActions();
 
 	/**
 	* Specify which MoviePipeline class type should be created by this executor when processing jobs.
@@ -247,7 +284,11 @@ protected:
 	virtual float GetStatusProgress_Implementation() const { return StatusProgress; }
 	virtual void CancelCurrentJob_Implementation() PURE_VIRTUAL(UMoviePipelineExecutorBase::CancelCurrentJobImpl, );
 	virtual void CancelAllJobs_Implementation() PURE_VIRTUAL(UMoviePipelineExecutorBase::CancelAllJobsImpl, );
+	UE_API virtual void SetPostRenderAction_Implementation(const EMoviePipelinePostRenderActionType PostRenderAction);
+	virtual EMoviePipelinePostRenderActionType GetPostRenderAction_Implementation() { return PostRenderActionType; }
+	virtual bool GetSupportsPostRenderActions_Implementation() { return false; }
 	// ~UMoviePipelineExecutorBase
+
 private:
 	/** 
 	* Called when the Executor has finished all jobs. Reports success if no jobs
@@ -277,6 +318,9 @@ private:
 	/** Set automatically when the error delegate gets broadcast (if fatal). */
 	bool bAnyJobHadFatalError;
 
+	/** The type of action that will be taken after all renders in the queue finish. */
+	EMoviePipelinePostRenderActionType PostRenderActionType;
+
 protected:
 	/** 
 	* Attempts to connect a socket to the specified host and port. This is a blocking call.
diff --git a/Engine/Plugins/MovieScene/MovieRenderPipeline/Source/MovieRenderPipelineCore/Public/MoviePipelinePostRenderSettings.h b/Engine/Plugins/MovieScene/MovieRenderPipeline/Source/MovieRenderPipelineCore/Public/MoviePipelinePostRenderSettings.h
--- a/Engine/Plugins/MovieScene/MovieRenderPipeline/Source/MovieRenderPipelineCore/Public/MoviePipelinePostRenderSettings.h
+++ b/Engine/Plugins/MovieScene/MovieRenderPipeline/Source/MovieRenderPipelineCore/Public/MoviePipelinePostRenderSettings.h
@@ -8,6 +8,9 @@
 
 #define UE_API MOVIERENDERPIPELINECORE_API
 
+// Note: There are a number of enums in here which are named with a MovieGraph prefix. This is no longer accurate; these are
+// now usable in the broader movie pipeline with the legacy system as well, but the naming is kept as-is for backwards compatibility.
+
 /** The type of viewer that should be used to play back media. */
 UENUM(BlueprintType)
 enum class EMovieGraphPlaybackMethod : uint8
@@ -99,6 +102,17 @@ enum class EMovieGraphOutputTypePlayback : uint8
 	PlayAllOutputTypes
 };
 
+/** When opening directories via the Play After Render feature, this dictates how they're opened. */
+UENUM(BlueprintType)
+enum class EMovieGraphDirectoryOpenBehavior : uint8
+{
+	/** Open the output directory for the first job in the queue only. */
+	OpenDirectoryForFirstJob,
+
+	/** Open the output directory for all jobs in the queue. */
+	OpenDirectoryForAllJobs
+};
+
 /** Options that dictate how videos should be played back after a render completes. */
 USTRUCT(BlueprintType)
 struct FMovieGraphPostRenderVideoPlayOptions
@@ -157,7 +171,7 @@ struct FMovieGraphPostRenderSettings
 	 * When deciding which media to play, media with extensions ranked higher in this list will play instead of media ranked lower.
 	 * For example, if EXR is listed above PNG, and both EXRs and PNGs were rendered, the EXRs will play instead of the PNGs.
 	 */
-	UPROPERTY(EditAnywhere, BlueprintReadWrite, Category = "Play Render Output (Quick Render only)")
+	UPROPERTY(EditAnywhere, BlueprintReadWrite, Category = "Play Render Output")
 	TArray<FString> OutputTypePriorityOrder;
 
 	/**
@@ -167,23 +181,27 @@ struct FMovieGraphPostRenderSettings
 	 * player as well (ie, two or more players will not be opened). If using the Operating System to open media, this setting is effectively ignored
 	 * because only the highest-ranked media in Output Type Priority Order will be opened.
 	 */
-	UPROPERTY(EditAnywhere, BlueprintReadWrite, Category = "Play Render Output (Quick Render only)")
+	UPROPERTY(EditAnywhere, BlueprintReadWrite, Category = "Play Render Output")
 	EMovieGraphOutputTypePlayback OutputTypePlayback;
 
+	/** When opening a directory instead of playing back media, this dictates which directories will be opened. */
+	UPROPERTY(EditAnywhere, BlueprintReadWrite, Category = "Play Render Output")
+	EMovieGraphDirectoryOpenBehavior DirectoryOpenBehavior;
+
 	/** These options are used to determine how to play the rendered images. */
-	UPROPERTY(EditAnywhere, BlueprintReadWrite, Category = "Play Render Output (Quick Render only)", DisplayName = "Image Sequences")
+	UPROPERTY(EditAnywhere, BlueprintReadWrite, Category = "Play Render Output", DisplayName = "Image Sequences")
 	FMovieGraphPostRenderImageSequencePlayOptions PostRenderImageSequencePlayOptions;
 
 	/** These options are used to determine how to play Apple ProRes movies. */
-	UPROPERTY(EditAnywhere, BlueprintReadWrite, Category = "Play Render Output (Quick Render only)", DisplayName = "Apple ProRes Movies")
+	UPROPERTY(EditAnywhere, BlueprintReadWrite, Category = "Play Render Output", DisplayName = "Apple ProRes Movies")
 	FMovieGraphPostRenderVideoPlayOptions PostRenderAppleProResPlayOptions;
 
 	/** These options are used to determine how to play Avid DNxHR movies. */
-	UPROPERTY(EditAnywhere, BlueprintReadWrite, Category = "Play Render Output (Quick Render only)", DisplayName = "Avid DNxHR Movies")
+	UPROPERTY(EditAnywhere, BlueprintReadWrite, Category = "Play Render Output", DisplayName = "Avid DNxHR Movies")
 	FMovieGraphPostRenderVideoPlayOptions PostRenderAvidDNxHRPlayOptions;
 
 	/** These options are used to determine how to play H.264 MP4 movies. */
-	UPROPERTY(EditAnywhere, BlueprintReadWrite, Category = "Play Render Output (Quick Render only)", DisplayName = "H.264 MP4 Movies")
+	UPROPERTY(EditAnywhere, BlueprintReadWrite, Category = "Play Render Output", DisplayName = "H.264 MP4 Movies")
 	FMovieGraphPostRenderVideoPlayOptions PostRenderMP4PlayOptions;
 
 	// TODO: No CLE support yet.
diff --git a/Engine/Plugins/MovieScene/MovieRenderPipeline/Source/MovieRenderPipelineCore/Public/MoviePipelineQueue.h b/Engine/Plugins/MovieScene/MovieRenderPipeline/Source/MovieRenderPipelineCore/Public/MoviePipelineQueue.h
--- a/Engine/Plugins/MovieScene/MovieRenderPipeline/Source/MovieRenderPipelineCore/Public/MoviePipelineQueue.h
+++ b/Engine/Plugins/MovieScene/MovieRenderPipeline/Source/MovieRenderPipelineCore/Public/MoviePipelineQueue.h
@@ -41,6 +41,23 @@ public:
 
 DECLARE_MULTICAST_DELEGATE_TwoParams(FOnMoviePipelineShotGraphPresetChanged, UMoviePipelineExecutorShot*, UMovieGraphConfig*);
 
+/** Project-based settings for the queue. */
+UCLASS(MinimalAPI, Config=EditorPerProjectUserSettings)
+class UMoviePipelineQueueProjectSettings : public UObject
+{
+	GENERATED_BODY()
+
+private:
+#if WITH_EDITOR
+	virtual void PostEditChangeProperty(FPropertyChangedEvent& PropertyChangedEvent) override;
+#endif	// WITH_EDITOR
+
+public:
+	/** When a queue finishes rendering, this determines the action that should be taken (if any). */
+	UPROPERTY(Config, BlueprintReadWrite, EditAnywhere, Category = "Movie Render Pipeline")
+	EMoviePipelinePostRenderActionType PostRenderActionType = EMoviePipelinePostRenderActionType::DoNothing;
+};
+
 /**
 * This class represents a segment of work within the Executor Job. This should be owned
 * by the UMoviePipelineExecutorJob and can be created before the movie pipeline starts to
diff --git a/Engine/Plugins/MovieScene/MovieRenderPipeline/Source/MovieRenderPipelineCore/Public/MoviePipelineUtils.h b/Engine/Plugins/MovieScene/MovieRenderPipeline/Source/MovieRenderPipelineCore/Public/MoviePipelineUtils.h
--- a/Engine/Plugins/MovieScene/MovieRenderPipeline/Source/MovieRenderPipelineCore/Public/MoviePipelineUtils.h
+++ b/Engine/Plugins/MovieScene/MovieRenderPipeline/Source/MovieRenderPipelineCore/Public/MoviePipelineUtils.h
@@ -263,6 +263,12 @@ namespace MoviePipeline
 {
 	MOVIERENDERPIPELINECORE_API void GetOutputStateFormatArgs(TMap<FString, FString>& InFilenameArguments, TMap<FString, FString>& InFileMetadata, const FString FrameNumber, const FString FrameNumberShot, const FString FrameNumberRel, const FString FrameNumberShotRel, const FString CameraName, const FString ShotName);
 
+	/**
+	 * Given a non-templated file name/path, returns the "frame-templated" file path. A frame-templated file path is one that has had its
+	 * {frame_*} token, if any, replaced with the {frame_placeholder} token. Returns the original file path if no modifications were done.
+	 */
+	MOVIERENDERPIPELINECORE_API FString GetFrameTemplatedFilePath(const FString& InFilePath);
+
 	/**
 	 * Gets the job variable assignments for a specific graph. Creates a new variable assignments container if one was not found for the given graph.
 	 * The owner of the variable assignments must be provided in case a new assignment needs to be created.
diff --git a/Engine/Plugins/MovieScene/MovieRenderPipeline/Source/MovieRenderPipelineCore/Public/MovieRenderPipelineDataTypes.h b/Engine/Plugins/MovieScene/MovieRenderPipeline/Source/MovieRenderPipelineCore/Public/MovieRenderPipelineDataTypes.h
--- a/Engine/Plugins/MovieScene/MovieRenderPipeline/Source/MovieRenderPipelineCore/Public/MovieRenderPipelineDataTypes.h
+++ b/Engine/Plugins/MovieScene/MovieRenderPipeline/Source/MovieRenderPipelineCore/Public/MovieRenderPipelineDataTypes.h
@@ -1347,6 +1347,16 @@ namespace MoviePipeline
 	};
 }
 
+/** A simple struct wrapper around a string array so it can be included in a TMap. */
+USTRUCT(BlueprintType)
+struct FMoviePipelineStringArray
+{
+	GENERATED_BODY()
+
+	UPROPERTY(VisibleAnywhere, BlueprintReadOnly, Category = "Movie Pipeline")
+	TArray<FString> Array;
+};
+
 USTRUCT(BlueprintType)
 struct FMoviePipelineRenderPassOutputData
 {
@@ -1354,6 +1364,18 @@ struct FMoviePipelineRenderPassOutputData
 
 	UPROPERTY(VisibleAnywhere, BlueprintReadOnly, Category = "Movie Pipeline")
 	TArray<FString> FilePaths;
+
+	/**
+	 * The "frame templated" file paths. This groups together a set of output frames by frame number, eg "Seq.Shot.{frame_placeholder}.exr".
+	 * If the render generated more than one output type (eg, EXR and JPEG), there may be more than one path here, but there should only
+	 * be one per output type.
+	 */
+	UPROPERTY(VisibleAnywhere, BlueprintReadOnly, Category = "Movie Pipeline")
+	TArray<FString> FrameTemplatedFilePaths;
+
+	/** Maps a setting type to all of the files that the setting generated. */
+	UPROPERTY(VisibleAnywhere, BlueprintReadOnly, Category = "Movie Pipeline")
+	TMap<FSoftClassPath, FMoviePipelineStringArray> SettingTypeToFilePaths;
 };
 
 USTRUCT(BlueprintType)
@@ -1379,11 +1401,18 @@ namespace MoviePipeline
 	{
 		FMoviePipelineOutputFutureData()
 			: Shot(nullptr)
+			, OriginSettingClass(nullptr)
 		{}
 
 		UMoviePipelineExecutorShot* Shot;
 		FString FilePath;
 		FMoviePipelinePassIdentifier PassIdentifier;
+
+		/** The "frame templated" filepath. This groups together a set of output frames by frame number, eg "Seq.Shot.{frame_placeholder}.exr". */
+		FString FrameTemplatedFilePath;
+
+		/** The setting class that this file originated from. */
+		UClass* OriginSettingClass;
 	};
 }
 
@@ -1458,6 +1487,22 @@ struct FMoviePipelineOutputData
 	TArray<FMovieGraphRenderOutputData> GraphData;
 };
 
+/** The action that should be taken after finishing the render. */
+UENUM(BlueprintType)
+enum class EMoviePipelinePostRenderActionType : uint8
+{
+	/** Do nothing after the render. */
+	DoNothing				= 0,
+
+	/** Open the media that was generated from the render. The application used for this is defined in Editor Preferences for Movie Render Graph. */
+	PlayRenderOutput		= 1 << 0,
+
+	/** Open the directory that contains the files generated by the render. */
+	OpenOutputDirectory		= 1 << 1
+};
+// Note: The values above are treated as flags in the possible future scenario where multiple actions can be taken
+// at the same time (and the actions could be OR'd together).
+
 /**
  * A pixel preprocessor for use with FImageWriteTask::PixelPreProcessor that does a simple alpha blend of the provided image onto the
  * target pixel data. This isn't very general purpose.
diff --git a/Engine/Plugins/MovieScene/MovieRenderPipeline/Source/MovieRenderPipelineEditor/Private/Graph/MovieGraphQuickRender.cpp b/Engine/Plugins/MovieScene/MovieRenderPipeline/Source/MovieRenderPipelineEditor/Private/Graph/MovieGraphQuickRender.cpp
--- a/Engine/Plugins/MovieScene/MovieRenderPipeline/Source/MovieRenderPipelineEditor/Private/Graph/MovieGraphQuickRender.cpp
+++ b/Engine/Plugins/MovieScene/MovieRenderPipeline/Source/MovieRenderPipelineEditor/Private/Graph/MovieGraphQuickRender.cpp
@@ -227,6 +227,7 @@ void UMovieGraphQuickRenderSubsystem::BeginQuickRender(const EMovieGraphQuickRen
 	// Always use the PIE executor for Quick Render, rather than the executor specified in the project settings for normal MRG.
 	TemporaryExecutor = NewObject<UMoviePipelinePIEExecutor>(GetTransientPackage());
 	TemporaryExecutor->SetAllowUsingUnsavedLevels(true);
+	TemporaryExecutor->SetPostRenderAction(QuickRenderModeSettings->PostRenderBehavior);
 
 	// Do any needed setup in the PIE world before rendering starts
 	TemporaryExecutor->OnIndividualJobStarted().AddLambda([this, EditorWorld](UMoviePipelineExecutorJob* InJob)
@@ -237,12 +238,6 @@ void UMovieGraphQuickRenderSubsystem::BeginQuickRender(const EMovieGraphQuickRen
 		// DEBUG_SaveAndOpenSequence(RenderingLevelSequence.Get(), TEXT("/Game/Temp/TempRenderingLevelSequence"));
 	});
 
-	// View the job's rendered frame(s)/video(s) if requested
-	TemporaryExecutor->OnIndividualJobWorkFinished().AddLambda([this](FMoviePipelineOutputData OutputData)
-	{
-		HandleJobFinished(QuickRenderModeSettings, OutputData);
-	});
-
 	// Clean up
 	TemporaryExecutor->OnExecutorFinished().AddLambda([this](UMoviePipelineExecutorBase* InExecutor, bool bSuccess)
 	{
@@ -1369,25 +1364,6 @@ void UMovieGraphQuickRenderSubsystem::PerformPreRenderSetup(UWorld* InEditorWorl
 	}
 }
 
-void UMovieGraphQuickRenderSubsystem::OpenPostRenderFileDisplayProcessor(const FMoviePipelineOutputData& InOutputData) const
-{
-	FMoviePipelinePostRenderFileDisplayProcessor PostRenderFileDisplayProcessor(GetDefault<UMovieRenderGraphEditorSettings>()->PostRenderSettings);
-	PostRenderFileDisplayProcessor.AddFiles(InOutputData);
-	PostRenderFileDisplayProcessor.OpenFiles();
-}
-
-void UMovieGraphQuickRenderSubsystem::HandleJobFinished(const UMovieGraphQuickRenderModeSettings* InQuickRenderSettings, const FMoviePipelineOutputData& InGeneratedOutputData)
-{
-	if (InQuickRenderSettings->PostRenderBehavior == EMovieGraphQuickRenderPostRenderActionType::PlayRenderOutput)
-	{
-		OpenPostRenderFileDisplayProcessor(InGeneratedOutputData);
-	}
-	else if (InQuickRenderSettings->PostRenderBehavior == EMovieGraphQuickRenderPostRenderActionType::OpenOutputDirectory)
-	{
-		OpenOutputDirectory(InQuickRenderSettings);
-	}
-}
-
 bool UMovieGraphQuickRenderSubsystem::IsViewportLookFlagActive(const EMovieGraphQuickRenderViewportLookFlags ViewportLookFlag) const
 {
 	if (!IsValid(QuickRenderModeSettings))
diff --git a/Engine/Plugins/MovieScene/MovieRenderPipeline/Source/MovieRenderPipelineEditor/Private/MoviePipelinePIEExecutor.cpp b/Engine/Plugins/MovieScene/MovieRenderPipeline/Source/MovieRenderPipelineEditor/Private/MoviePipelinePIEExecutor.cpp
--- a/Engine/Plugins/MovieScene/MovieRenderPipeline/Source/MovieRenderPipelineEditor/Private/MoviePipelinePIEExecutor.cpp
+++ b/Engine/Plugins/MovieScene/MovieRenderPipeline/Source/MovieRenderPipelineEditor/Private/MoviePipelinePIEExecutor.cpp
@@ -6,12 +6,14 @@
 #include "MoviePipeline.h"
 #include "MoviePipelineGameOverrideSetting.h"
 #include "MoviePipelineOutputSetting.h"
+#include "MoviePipelinePostRenderFileDisplayProcessor.h"
 #include "MovieRenderPipelineCoreModule.h"
 #include "Framework/Application/SlateApplication.h"
 #include "MovieRenderPipelineDataTypes.h"
 #include "MovieRenderPipelineSettings.h"
 #include "Editor/EditorEngine.h"
 #include "Editor.h"
+#include "HAL/FileManager.h"
 #include "Settings/LevelEditorPlaySettings.h"
 #include "MoviePipelineQueue.h"
 #include "MoviePipelineBlueprintLibrary.h"
@@ -23,6 +25,7 @@
 #include "Logging/MessageLog.h"
 #include "Graph/MovieGraphBlueprintLibrary.h"
 #include "Graph/MovieGraphPipeline.h"
+#include "Graph/MovieRenderGraphEditorSettings.h"
 #include "Graph/Nodes/MovieGraphGlobalGameOverrides.h"
 #include "LevelSequence.h"
 
@@ -407,6 +410,58 @@ void UMoviePipelinePIEExecutor::BeginDestroy()
 	Super::BeginDestroy();
 }
 
+void UMoviePipelinePIEExecutor::OnExecutorFinishedImpl()
+{
+	Super::OnExecutorFinishedImpl();
+
+	const UMovieRenderGraphEditorSettings* EditorSettings = GetDefault<UMovieRenderGraphEditorSettings>();
+
+	// Now that the queue has finished, play the rendered frames in an external application, open the output directory,
+	// or do nothing.
+	if (GetPostRenderAction() == EMoviePipelinePostRenderActionType::PlayRenderOutput)
+	{
+		FMoviePipelinePostRenderFileDisplayProcessor PostRenderFileDisplayProcessor(EditorSettings->PostRenderSettings);
+
+		for (const FMoviePipelineOutputData& OutputData : AllJobOutputData)
+		{
+			PostRenderFileDisplayProcessor.AddFiles(OutputData);
+		}
+		
+		PostRenderFileDisplayProcessor.OpenFiles();
+	}
+	else if (GetPostRenderAction() == EMoviePipelinePostRenderActionType::OpenOutputDirectory)
+	{
+		while (!ResolvedOutputDirectories.IsEmpty())
+		{
+			const FString ResolvedOutputDirectory = ResolvedOutputDirectories[0];
+			ResolvedOutputDirectories.RemoveAt(0);
+
+			if (!ResolvedOutputDirectory.IsEmpty())
+			{
+				// The directory might not exist yet. Create it (if needed) so ExploreFolder() can open it.
+				IFileManager::Get().MakeDirectory(*ResolvedOutputDirectory, true);
+
+				FPlatformProcess::ExploreFolder(*ResolvedOutputDirectory);
+			}
+
+			// Depending on the settings, we might only need to open one directory.
+			if (EditorSettings->PostRenderSettings.DirectoryOpenBehavior == EMovieGraphDirectoryOpenBehavior::OpenDirectoryForFirstJob)
+			{
+				break;
+			}
+		}
+	}
+
+	// Reset all cached data tracking rendered output
+	ResolvedOutputDirectories.Empty();
+	AllJobOutputData.Empty();
+}
+
+bool UMoviePipelinePIEExecutor::GetSupportsPostRenderActions_Implementation()
+{
+	return true;
+}
+
 void UMoviePipelinePIEExecutor::OnPIEEnded(bool)
 {
 	FEditorDelegates::EndPIE.RemoveAll(this);
@@ -458,11 +513,16 @@ void UMoviePipelinePIEExecutor::DelayedFinishNotification()
 	// Now that another frame has passed and we should be OK to start another PIE session, notify our owner.
 	OnIndividualPipelineFinished(nullptr);
 }
+
 void UMoviePipelinePIEExecutor::OnIndividualJobStartedImpl(UMoviePipelineExecutorJob* InJob)
 {
 	// Broadcast to both Native and Python/BP
 	OnIndividualJobStartedDelegateNative.Broadcast(InJob);
 	OnIndividualJobStartedDelegate.Broadcast(InJob);
+
+	// By the time the job finishes, it will have most likely been garbage collected (at least when using the graph, since jobs that use the graph
+	// are duplicated). Instead of caching the resolved output directories when the job finishes, cache it when starting while the job is still valid.
+	ResolvedOutputDirectories.AddUnique(UMoviePipelineEditorBlueprintLibrary::ResolveOutputDirectoryFromJob(InJob));
 }
 
 void UMoviePipelinePIEExecutor::OnIndividualJobFinishedImpl(FMoviePipelineOutputData InOutputData)
@@ -475,6 +535,11 @@ void UMoviePipelinePIEExecutor::OnIndividualJobFinishedImpl(FMoviePipelineOutput
 
 	OnIndividualJobWorkFinishedDelegateNative.Broadcast(InOutputData);
 	OnIndividualJobWorkFinishedDelegate.Broadcast(InOutputData);
+
+	// Null out the pipeline and job pointers in the output data before caching (so they are GC'd correctly after the pipeline finishes)
+	InOutputData.Pipeline = nullptr;
+	InOutputData.Job = nullptr;
+	AllJobOutputData.Add(InOutputData);
 }
 
 #undef LOCTEXT_NAMESPACE // "MoviePipelinePIEExecutor"
diff --git a/Engine/Plugins/MovieScene/MovieRenderPipeline/Source/MovieRenderPipelineEditor/Private/MoviePipelinePostRenderFileDisplayProcessor.cpp b/Engine/Plugins/MovieScene/MovieRenderPipeline/Source/MovieRenderPipelineEditor/Private/MoviePipelinePostRenderFileDisplayProcessor.cpp
--- a/Engine/Plugins/MovieScene/MovieRenderPipeline/Source/MovieRenderPipelineEditor/Private/MoviePipelinePostRenderFileDisplayProcessor.cpp
+++ b/Engine/Plugins/MovieScene/MovieRenderPipeline/Source/MovieRenderPipelineEditor/Private/MoviePipelinePostRenderFileDisplayProcessor.cpp
@@ -96,7 +96,8 @@ FMoviePipelinePostRenderFileDisplayProcessor::FFilesToOpen FMoviePipelinePostRen
 {
 	FFilesToOpen FilesToOpen;
 	
-	if (InPipelineOutputData.GraphData.IsEmpty())
+	// Nothing to do if the graph (GraphData) and preset (ShotData) output is empty.
+	if (InPipelineOutputData.GraphData.IsEmpty() && InPipelineOutputData.ShotData.IsEmpty())
 	{
 		return FilesToOpen;
 	}
@@ -329,6 +330,22 @@ void FMoviePipelinePostRenderFileDisplayProcessor::GroupFilesByShot(
 	TArray<TMap<FString, TArray<FString>>>& InImagesGroupedByShot,
 	TArray<TMap<FString, TArray<FString>>>& InFrameTemplateImagesGroupedByShot,
 	TArray<TMap<FString, TArray<FString>>>& InVideosGroupedByShot) const
+{
+	if (!InPipelineOutputData.GraphData.IsEmpty())
+	{
+		GroupFilesByShot_Graph(InPipelineOutputData, InImagesGroupedByShot, InFrameTemplateImagesGroupedByShot, InVideosGroupedByShot);
+	}
+	else
+	{
+		GroupFilesByShot_Preset(InPipelineOutputData, InImagesGroupedByShot, InFrameTemplateImagesGroupedByShot, InVideosGroupedByShot);
+	}
+}
+
+void FMoviePipelinePostRenderFileDisplayProcessor::GroupFilesByShot_Graph(
+	FMoviePipelineOutputData& InPipelineOutputData,
+	TArray<TMap<FString, TArray<FString>>>& InImagesGroupedByShot,
+	TArray<TMap<FString, TArray<FString>>>& InFrameTemplateImagesGroupedByShot,
+	TArray<TMap<FString, TArray<FString>>>& InVideosGroupedByShot) const
 {
 	// OutputData is an individual shot's collection of rendered output (images + movies)
 	for (FMovieGraphRenderOutputData& OutputData : InPipelineOutputData.GraphData)
@@ -443,6 +460,68 @@ void FMoviePipelinePostRenderFileDisplayProcessor::GroupFilesByShot(
 	}
 }
 
+void FMoviePipelinePostRenderFileDisplayProcessor::GroupFilesByShot_Preset(
+	FMoviePipelineOutputData& InPipelineOutputData,
+	TArray<TMap<FString, TArray<FString>>>& InImagesGroupedByShot,
+	TArray<TMap<FString, TArray<FString>>>& InFrameTemplateImagesGroupedByShot,
+	TArray<TMap<FString, TArray<FString>>>& InVideosGroupedByShot) const
+{
+	// OutputData is an individual shot's collection of rendered output (images + movies)
+	for (FMoviePipelineShotOutputData& OutputData : InPipelineOutputData.ShotData)
+	{
+		TMap<FString, TArray<FString>>& ImagesByExtension = InImagesGroupedByShot.AddDefaulted_GetRef();
+		TMap<FString, TArray<FString>>& ImagesByFrameTemplate = InFrameTemplateImagesGroupedByShot.AddDefaulted_GetRef();
+		TMap<FString, TArray<FString>>& VideosByExtension = InVideosGroupedByShot.AddDefaulted_GetRef();
+
+		// Add the images/videos for each render pass within the shot
+		for (const TPair<FMoviePipelinePassIdentifier, FMoviePipelineRenderPassOutputData>& PassData : OutputData.RenderPassData)
+		{
+			// For each setting type that generated files, add its files
+			for (const TPair<FSoftClassPath, FMoviePipelineStringArray>& FilesFromSetting : PassData.Value.SettingTypeToFilePaths)
+			{
+				const UClass* SettingClass = FilesFromSetting.Key.ResolveClass();
+				if (!SettingClass)
+				{
+					UE_LOG(LogMovieRenderPipeline, Warning, TEXT("Found an invalid setting type while determining the media to display post-render. This setting's media will be skipped."));
+					continue;
+				}
+				
+				const bool bIsImageNode = SettingClass->IsChildOf(UMoviePipelineImageSequenceOutputBase::StaticClass());
+
+				// For all files generated by the setting, categorize it either by extension and/or frame template (for images)
+				for (const FString& SettingFile : FilesFromSetting.Value.Array)
+				{
+					const FString Extension = FPaths::GetExtension(SettingFile).ToUpper();
+					if (Extension.IsEmpty())
+					{
+						continue;
+					}
+					
+					if (bIsImageNode)
+					{
+						ImagesByExtension.FindOrAdd(Extension).Add(SettingFile);
+
+						// The render pass will contain a set of frame templated file paths, in addition to normal file paths. Determine if this
+						// file should be associated with any of these frame templated paths (matched by extension).
+						for (const FString& FrameTemplatedFilePath : PassData.Value.FrameTemplatedFilePaths)
+						{
+							if (FPaths::GetExtension(FrameTemplatedFilePath).ToUpper() == Extension)
+							{
+								ImagesByFrameTemplate.FindOrAdd(FrameTemplatedFilePath).Add(SettingFile);
+								break;
+							}
+						}
+					}
+					else
+					{
+						VideosByExtension.FindOrAdd(Extension).Add(SettingFile);
+					}
+				}
+			}
+		}
+	}
+}
+
 void FMoviePipelinePostRenderFileDisplayProcessor::GetHighestPriorityExtensions(
 	const TArray<TMap<FString, TArray<FString>>>& InImagesGroupedByShot,
 	const TArray<TMap<FString, TArray<FString>>>& InVideosGroupedByShot,
@@ -495,11 +574,11 @@ void FMoviePipelinePostRenderFileDisplayProcessor::LaunchFilesWithSettings(const
 	const FString StartFrame = InFrameRangeToOpen.Key;
 	const FString EndFrame = InFrameRangeToOpen.Value;
 	
-	auto ReplaceFramePlaceholderInFilePaths = [&StartFrame](TArray<FString>& InFilePaths)
+	auto ReplaceFramePlaceholderInFilePaths = [this, &StartFrame](TArray<FString>& InFilePaths)
 	{
 		FString FramePlaceholder;
 		
-		const EMovieGraphFrameRangeNotation FrameRangeNotationType = GetDefault<UMovieRenderGraphEditorSettings>()->PostRenderSettings.PostRenderImageSequencePlayOptions.FrameRangeNotation;
+		const EMovieGraphFrameRangeNotation FrameRangeNotationType = PostRenderSettings.PostRenderImageSequencePlayOptions.FrameRangeNotation;
 		switch (FrameRangeNotationType)
 		{
 		case EMovieGraphFrameRangeNotation::Hash:
@@ -532,16 +611,18 @@ void FMoviePipelinePostRenderFileDisplayProcessor::LaunchFilesWithSettings(const
 	{
 		if (!StartFrame.IsEmpty() && !EndFrame.IsEmpty())
 		{
-			const FString FrameRangeString = FString::Format(TEXT("{0}-{1}"), { StartFrame, EndFrame });
-
 			const EMovieGraphFrameRangeNotation FrameRangeNotation = PostRenderSettings.PostRenderImageSequencePlayOptions.FrameRangeNotation;
 			if (FrameRangeNotation == EMovieGraphFrameRangeNotation::HashWithStartEndFrame)
 			{
+				const FString FrameRangeString = FString::Format(TEXT("{0}-{1}"), { StartFrame, EndFrame });
+				
 				// If using # with frame ranges, specify the frame range after the file path(s). This is how RV expects it.
 				InConcatenatedFilePaths = FString::Format(TEXT("{0} {1}"), { InConcatenatedFilePaths, FrameRangeString });
 			}
 			else if (FrameRangeNotation == EMovieGraphFrameRangeNotation::DollarF)
 			{
+				const FString FrameRangeString = FString::Format(TEXT("{0} {1}"), { StartFrame, EndFrame });
+				
 				// $F is used exclusively by MPlay, so provide the frame range via -f with a 1-frame step (eg, "-f 0 150 1")
 				InConcatenatedFilePaths = FString::Format(TEXT("-f {0} 1 {1}"), { FrameRangeString, InConcatenatedFilePaths });
 			}
diff --git a/Engine/Plugins/MovieScene/MovieRenderPipeline/Source/MovieRenderPipelineEditor/Private/MoviePipelinePostRenderFileDisplayProcessor.h b/Engine/Plugins/MovieScene/MovieRenderPipeline/Source/MovieRenderPipelineEditor/Private/MoviePipelinePostRenderFileDisplayProcessor.h
--- a/Engine/Plugins/MovieScene/MovieRenderPipeline/Source/MovieRenderPipelineEditor/Private/MoviePipelinePostRenderFileDisplayProcessor.h
+++ b/Engine/Plugins/MovieScene/MovieRenderPipeline/Source/MovieRenderPipelineEditor/Private/MoviePipelinePostRenderFileDisplayProcessor.h
@@ -29,7 +29,7 @@ public:
 
 	/**
 	 * Adds a set of files to the processor. This should only be called once per job, otherwise duplicate data may be added.
-	 * Only graph-related output data is added.
+	 * Both graph and preset output data is supported.
 	 */
 	void AddFiles(const FMoviePipelineOutputData& InRenderOutputData);
 
@@ -81,6 +81,20 @@ private:
 		TArray<TMap<FString, TArray<FString>>>& InFrameTemplateImagesGroupedByShot,
 		TArray<TMap<FString, TArray<FString>>>& InVideosGroupedByShot) const;
 
+	/** Helper that implements graph-based logic for GroupFilesByShot(). */
+	void GroupFilesByShot_Graph(
+		FMoviePipelineOutputData& InPipelineOutputData,
+		TArray<TMap<FString, TArray<FString>>>& InImagesGroupedByShot,
+		TArray<TMap<FString, TArray<FString>>>& InFrameTemplateImagesGroupedByShot,
+		TArray<TMap<FString, TArray<FString>>>& InVideosGroupedByShot) const;
+
+	/** Helper that implements preset-based logic for GroupFilesByShot(). */
+	void GroupFilesByShot_Preset(
+		FMoviePipelineOutputData& InPipelineOutputData,
+		TArray<TMap<FString, TArray<FString>>>& InImagesGroupedByShot,
+		TArray<TMap<FString, TArray<FString>>>& InFrameTemplateImagesGroupedByShot,
+		TArray<TMap<FString, TArray<FString>>>& InVideosGroupedByShot) const;
+
 	void GetHighestPriorityExtensions(
 		const TArray<TMap<FString, TArray<FString>>>& InImagesGroupedByShot,
 		const TArray<TMap<FString, TArray<FString>>>& InVideosGroupedByShot,
diff --git a/Engine/Plugins/MovieScene/MovieRenderPipeline/Source/MovieRenderPipelineEditor/Private/MoviePipelineQuickRenderMenu.cpp b/Engine/Plugins/MovieScene/MovieRenderPipeline/Source/MovieRenderPipelineEditor/Private/MoviePipelineQuickRenderMenu.cpp
--- a/Engine/Plugins/MovieScene/MovieRenderPipeline/Source/MovieRenderPipelineEditor/Private/MoviePipelineQuickRenderMenu.cpp
+++ b/Engine/Plugins/MovieScene/MovieRenderPipeline/Source/MovieRenderPipelineEditor/Private/MoviePipelineQuickRenderMenu.cpp
@@ -428,7 +428,7 @@ void FMoviePipelineQuickRenderMenu::GenerateQuickRenderConfigurationMenuSection(
 					: LOCTEXT("QuickRenderConfigMenu_InvalidGraph", "Invalid");
 				
 				// Generate "After Render" text
-				const FText AfterRender = (QuickRenderModeSettings->PostRenderBehavior == EMovieGraphQuickRenderPostRenderActionType::PlayRenderOutput)
+				const FText AfterRender = (QuickRenderModeSettings->PostRenderBehavior == EMoviePipelinePostRenderActionType::PlayRenderOutput)
 					? LOCTEXT("QuickRenderConfigMenu_PlayRenderOutput", "Play Render Output")
 					: LOCTEXT("QuickRenderConfigMenu_DoNothing", "Do Nothing");
 				
diff --git a/Engine/Plugins/MovieScene/MovieRenderPipeline/Source/MovieRenderPipelineEditor/Private/MovieRenderPipelineEditorUtils.cpp b/Engine/Plugins/MovieScene/MovieRenderPipeline/Source/MovieRenderPipelineEditor/Private/MovieRenderPipelineEditorUtils.cpp
--- a/Engine/Plugins/MovieScene/MovieRenderPipeline/Source/MovieRenderPipelineEditor/Private/MovieRenderPipelineEditorUtils.cpp
+++ b/Engine/Plugins/MovieScene/MovieRenderPipeline/Source/MovieRenderPipelineEditor/Private/MovieRenderPipelineEditorUtils.cpp
@@ -24,7 +24,11 @@ namespace UE::MovieRenderPipelineEditor::Private
 		const TSubclassOf<UMoviePipelineExecutorBase> ExecutorClass = ProjectSettings->DefaultLocalExecutor.TryLoadClass<UMoviePipelineExecutorBase>();
 		check(ExecutorClass != nullptr);
 		
-		Subsystem->RenderQueueWithExecutor(ExecutorClass);
+		UMoviePipelineExecutorBase* Executor = Subsystem->RenderQueueWithExecutor(ExecutorClass);
+
+		// Tell the executor to run the post-render action. Some executors may not support this.
+		const UMoviePipelineQueueProjectSettings* QueueProjectSettings = GetDefault<UMoviePipelineQueueProjectSettings>();
+		Executor->SetPostRenderAction(QueueProjectSettings->PostRenderActionType);
 
 		constexpr bool bIsLocal = true;
 		FMoviePipelineTelemetry::SendRendersRequestedTelemetry(bIsLocal, Subsystem->GetQueue()->GetJobs());
diff --git a/Engine/Plugins/MovieScene/MovieRenderPipeline/Source/MovieRenderPipelineEditor/Private/Widgets/SMoviePipelineQueuePanel.cpp b/Engine/Plugins/MovieScene/MovieRenderPipeline/Source/MovieRenderPipelineEditor/Private/Widgets/SMoviePipelineQueuePanel.cpp
--- a/Engine/Plugins/MovieScene/MovieRenderPipeline/Source/MovieRenderPipelineEditor/Private/Widgets/SMoviePipelineQueuePanel.cpp
+++ b/Engine/Plugins/MovieScene/MovieRenderPipeline/Source/MovieRenderPipelineEditor/Private/Widgets/SMoviePipelineQueuePanel.cpp
@@ -44,6 +44,8 @@
 #include "Editor.h"
 #include "PropertyEditorModule.h"
 #include "IDetailsView.h"
+#include "ISettingsModule.h"
+#include "ISinglePropertyView.h"
 #include "Misc/MessageDialog.h"
 #include "Modules/ModuleManager.h"
 
@@ -55,7 +57,6 @@
 #include "AssetToolsModule.h"
 #include "Misc/FileHelper.h"
 
-
 #define LOCTEXT_NAMESPACE "SMoviePipelineQueuePanel"
 
 UE_DISABLE_OPTIMIZATION_SHIP
@@ -88,20 +89,34 @@ void SMoviePipelineQueuePanel::Construct(const FArguments& InArgs)
 		.OnPresetChosen(this, &SMoviePipelineQueuePanel::OnJobPresetChosen)
 		.OnJobSelectionChanged(this, &SMoviePipelineQueuePanel::OnSelectionChanged);
 
-	{
-		// Automatically select the first job in the queue
-		UMoviePipelineQueueSubsystem* Subsystem = GEditor->GetEditorSubsystem<UMoviePipelineQueueSubsystem>();
-		check(Subsystem);
+	const UMoviePipelineQueueSubsystem* Subsystem = GEditor->GetEditorSubsystem<UMoviePipelineQueueSubsystem>();
+	check(Subsystem);
+	UMoviePipelineQueue* Queue = Subsystem->GetQueue();
+	check(Queue);
 
-		TArray<UMoviePipelineExecutorJob*> Jobs;
-		if (Subsystem->GetQueue()->GetJobs().Num() > 0)
+	// Automatically select the first job in the queue
+	{
+		const TArray<UMoviePipelineExecutorJob*> CurrentJobs = Queue->GetJobs();
+		if (!CurrentJobs.IsEmpty())
 		{
-			Jobs.Add(Subsystem->GetQueue()->GetJobs()[0]);
+			// Go through the UI so it updates the UI selection too and then this will loop back
+			// around to OnSelectionChanged to update ourself.
+			const TArray<UMoviePipelineExecutorJob*> SelectedJobs = { CurrentJobs[0] };
+			PipelineQueueEditorWidget->SetSelectedJobs(SelectedJobs);
 		}
+	}
 
-		// Go through the UI so it updates the UI selection too and then this will loop back
-		// around to OnSelectionChanged to update ourself.
-		PipelineQueueEditorWidget->SetSelectedJobs(Jobs);
+	// Create a drop-down widget for the "After Render" functionality.
+	TSharedPtr<ISinglePropertyView> PostRenderBehaviorWidget;
+	{
+		UMoviePipelineQueueProjectSettings* QueueProjectSettings = GetMutableDefault<UMoviePipelineQueueProjectSettings>();
+		
+		FSinglePropertyParams PostRenderActionTypeParams;
+		PostRenderActionTypeParams.NamePlacement = EPropertyNamePlacement::Hidden;
+		PostRenderActionTypeParams.bHideResetToDefault = true;
+		
+		PostRenderBehaviorWidget = PropertyEditorModule.CreateSingleProperty(
+			QueueProjectSettings, GET_MEMBER_NAME_CHECKED(UMoviePipelineQueueProjectSettings, PostRenderActionType), PostRenderActionTypeParams);
 	}
 
 	ChildSlot
@@ -234,6 +249,48 @@ void SMoviePipelineQueuePanel::Construct(const FArguments& InArgs)
 			[
 				SNew(SHorizontalBox)
 
+				+ SHorizontalBox::Slot()
+				.VAlign(VAlign_Center)
+				.HAlign(HAlign_Left)
+				.Padding(5, 0, 0, 0)
+				.AutoWidth()
+				[
+					SNew(STextBlock)
+					.Text(LOCTEXT("AfterRender_Text", "After Render:"))
+				]
+
+				+ SHorizontalBox::Slot()
+				.VAlign(VAlign_Center)
+				.HAlign(HAlign_Left)
+				.AutoWidth()
+				[
+					PostRenderBehaviorWidget.ToSharedRef()
+				]
+
+				+ SHorizontalBox::Slot()
+				.VAlign(VAlign_Center)
+				.HAlign(HAlign_Left)
+				.AutoWidth()
+				[
+					SNew(SButton)
+					.ButtonStyle(FAppStyle::Get(), TEXT("SimpleButton"))
+					.ContentPadding(0)
+					.OnClicked_Lambda([]()
+					{
+						FModuleManager::LoadModuleChecked<ISettingsModule>("Settings")
+							.ShowViewer("Editor", "Plugins", "MovieRenderGraphEditorSettings");
+
+						return FReply::Handled();
+					})
+					.ToolTipText(LOCTEXT("OpenPlaybackPrefs_Tooltip", "Open up the preferences which dictate how media is opened after a render completes."))
+					[
+						SNew(SImage)
+						.ColorAndOpacity(FSlateColor::UseForeground())
+						.Image(FAppStyle::Get().GetBrush("EditorPreferences.TabIcon"))
+						.DesiredSizeOverride(FVector2D(16.f, 16.f))
+					]
+				]
+
 				+ SHorizontalBox::Slot()
 				.VAlign(VAlign_Fill)
 				.HAlign(HAlign_Left)
diff --git a/Engine/Plugins/MovieScene/MovieRenderPipeline/Source/MovieRenderPipelineEditor/Public/Graph/MovieGraphQuickRender.h b/Engine/Plugins/MovieScene/MovieRenderPipeline/Source/MovieRenderPipelineEditor/Public/Graph/MovieGraphQuickRender.h
--- a/Engine/Plugins/MovieScene/MovieRenderPipeline/Source/MovieRenderPipelineEditor/Public/Graph/MovieGraphQuickRender.h
+++ b/Engine/Plugins/MovieScene/MovieRenderPipeline/Source/MovieRenderPipelineEditor/Public/Graph/MovieGraphQuickRender.h
@@ -158,12 +158,6 @@ private:
 	/** Does any setup needed in the PIE world prior to the render starting. */
 	UE_API void PerformPreRenderSetup(UWorld* InEditorWorld);
 
-	/** Opens the rendered files that are in the given output data. */
-	UE_API void OpenPostRenderFileDisplayProcessor(const FMoviePipelineOutputData& InOutputData) const;
-
-	/** Takes care of post-render tasks, like opening the rendered files. */
-	UE_API void HandleJobFinished(const UMovieGraphQuickRenderModeSettings* InQuickRenderSettings, const FMoviePipelineOutputData& InGeneratedOutputData);
-
 	/** Determines if the given viewport look flag is currently active. */
 	UE_API bool IsViewportLookFlagActive(const EMovieGraphQuickRenderViewportLookFlags ViewportLookFlag) const;
 
diff --git a/Engine/Plugins/MovieScene/MovieRenderPipeline/Source/MovieRenderPipelineEditor/Public/MoviePipelinePIEExecutor.h b/Engine/Plugins/MovieScene/MovieRenderPipeline/Source/MovieRenderPipelineEditor/Public/MoviePipelinePIEExecutor.h
--- a/Engine/Plugins/MovieScene/MovieRenderPipeline/Source/MovieRenderPipelineEditor/Public/MoviePipelinePIEExecutor.h
+++ b/Engine/Plugins/MovieScene/MovieRenderPipeline/Source/MovieRenderPipelineEditor/Public/MoviePipelinePIEExecutor.h
@@ -78,9 +78,16 @@ public:
 	UE_API void SetAllowUsingUnsavedLevels(const bool bShouldAllow);
 
 protected:
+	//~ Begin UMoviePipelineLinearExecutorBase interface
 	UE_API virtual void Start(const UMoviePipelineExecutorJob* InJob) override;
-
 	UE_API virtual void BeginDestroy() override;
+	UE_API virtual void OnExecutorFinishedImpl() override;
+	//~ End UMoviePipelineLinearExecutorBase interface
+
+	//~ Begin UMoviePipelineExecutorBase interface
+	virtual bool GetSupportsPostRenderActions_Implementation() override;
+	//~ End UMoviePipelineExecutorBase interface
+	
 	/**
 	* This should be called after PIE has been shut down for an individual job and it is generally safer to
 	* make modifications to the editor world.
@@ -118,6 +125,12 @@ private:
 	/** Whether unsaved levels can be rendered from. */
 	bool bAllowUnsavedLevels;
 
+	/** Output data generated from all jobs in the queue (one entry in the array for each job). May be a mix of graph and preset output data. */
+	TArray<FMoviePipelineOutputData> AllJobOutputData;
+
+	/** The resolved output directories for all jobs that have been rendered. */
+	TArray<FString> ResolvedOutputDirectories;
+
 	/** Deprecated. use OnIndividualJobWorkFinishedDelegate instead. */
 	UE_DEPRECATED(4.27, "Use OnIndividualJobWorkFinishedDelegate instead.")
 	UPROPERTY(BlueprintAssignable, Category = "Movie Render Pipeline")
diff --git a/Engine/Plugins/MovieScene/MovieRenderPipeline/Source/MovieRenderPipelineRenderPasses/Private/MovieGraphImageSequenceOutputNode.cpp b/Engine/Plugins/MovieScene/MovieRenderPipeline/Source/MovieRenderPipelineRenderPasses/Private/MovieGraphImageSequenceOutputNode.cpp
--- a/Engine/Plugins/MovieScene/MovieRenderPipeline/Source/MovieRenderPipelineRenderPasses/Private/MovieGraphImageSequenceOutputNode.cpp
+++ b/Engine/Plugins/MovieScene/MovieRenderPipeline/Source/MovieRenderPipelineRenderPasses/Private/MovieGraphImageSequenceOutputNode.cpp
@@ -180,25 +180,7 @@ FString UMovieGraphImageSequenceOutputNode::CreateFileName(
 FString UMovieGraphImageSequenceOutputNode::GetFrameTemplatedFileName(
 	const FMovieGraphFilenameResolveParams& InParams, const FString& InFileNameFormatString, FMovieGraphResolveArgs& OutMergedFormatArgs) const
 {
-	const FString FramePlaceholder = TEXT("{frame_placeholder}");
-
-	FString FrameTemplatedFormatString = InFileNameFormatString;
-	if (FrameTemplatedFormatString.Contains(TEXT("{frame_number}")))
-	{
-		FrameTemplatedFormatString = InFileNameFormatString.Replace(TEXT("{frame_number}"), *FramePlaceholder);
-	}
-	else if (FrameTemplatedFormatString.Contains(TEXT("{frame_number_rel}")))
-	{
-		FrameTemplatedFormatString = InFileNameFormatString.Replace(TEXT("{frame_number_rel}"), *FramePlaceholder);
-	}
-	else if (FrameTemplatedFormatString.Contains(TEXT("{frame_number_shot}")))
-	{
-		FrameTemplatedFormatString = InFileNameFormatString.Replace(TEXT("{frame_number_shot}"), *FramePlaceholder);
-	}
-	else if (FrameTemplatedFormatString.Contains(TEXT("{frame_number_shot_rel}")))
-	{
-		FrameTemplatedFormatString = InFileNameFormatString.Replace(TEXT("{frame_number_shot_rel}"), *FramePlaceholder);
-	}
+	const FString FrameTemplatedFormatString = MoviePipeline::GetFrameTemplatedFilePath(InFileNameFormatString);
 
 	// If time dilation is being used, the parameters will ask ResolveFilenameFormatArguments to warn the user if there's no _rel frame number
 	// found, but we're intentionally overriding them above to be able to replace them with a completely unrelated token, so we don't actually
diff --git a/Engine/Plugins/MovieScene/MovieRenderPipeline/Source/MovieRenderPipelineRenderPasses/Private/MoviePipelineEXROutput.cpp b/Engine/Plugins/MovieScene/MovieRenderPipeline/Source/MovieRenderPipelineRenderPasses/Private/MoviePipelineEXROutput.cpp
--- a/Engine/Plugins/MovieScene/MovieRenderPipeline/Source/MovieRenderPipelineRenderPasses/Private/MoviePipelineEXROutput.cpp
+++ b/Engine/Plugins/MovieScene/MovieRenderPipeline/Source/MovieRenderPipelineRenderPasses/Private/MoviePipelineEXROutput.cpp
@@ -717,6 +717,7 @@ void UMoviePipelineImageSequenceOutput_EXR::OnReceiveImageDataImpl(FMoviePipelin
 
 	// We need to resolve the filename format string. We combine the folder and file name into one long string first
 	FString FinalFilePath;
+	FString FrameTemplatedFilePath;
 	FMoviePipelineFormatArgs FinalFormatArgs;
 	FString FinalImageSequenceFileName;
 	FString ClipName;
@@ -747,6 +748,15 @@ void UMoviePipelineImageSequenceOutput_EXR::OnReceiveImageDataImpl(FMoviePipelin
 			FinalFilePath = FPaths::ConvertRelativePathToFull(FinalFilePath);
 		}
 
+		// Re-resolve using the frame-templated format string
+		const FString FrameTemplatedFormatString = MoviePipeline::GetFrameTemplatedFilePath(FilePathFormatString);	
+		GetPipeline()->ResolveFilenameFormatArguments(FrameTemplatedFormatString, FormatOverrides, FrameTemplatedFilePath, FinalFormatArgs, &InMergedOutputFrame->FrameOutputState);
+
+		if (FPaths::IsRelative(FrameTemplatedFilePath))
+		{
+			FrameTemplatedFilePath = FPaths::ConvertRelativePathToFull(FrameTemplatedFilePath);
+		}
+
 		// Create a deterministic clipname by removing frame numbers, file extension, and any trailing .'s
 		UE::MoviePipeline::RemoveFrameNumberFormatStrings(FileNameFormatString, true);
 		GetPipeline()->ResolveFilenameFormatArguments(FileNameFormatString, FormatOverrides, ClipName, FinalFormatArgs, &InMergedOutputFrame->FrameOutputState);
@@ -837,7 +847,9 @@ void UMoviePipelineImageSequenceOutput_EXR::OnReceiveImageDataImpl(FMoviePipelin
 	MoviePipeline::FMoviePipelineOutputFutureData OutputData;
 	OutputData.Shot = GetPipeline()->GetActiveShotList()[ShotIndex];
 	OutputData.PassIdentifier = FMoviePipelinePassIdentifier(TEXT("")); // exrs put all the render passes internally so this resolves to a ""
+	OutputData.OriginSettingClass = GetClass();
 	OutputData.FilePath = FinalFilePath;
+	OutputData.FrameTemplatedFilePath = FrameTemplatedFilePath;
 	GetPipeline()->AddOutputFuture(ImageWriteQueue->Enqueue(MoveTemp(MultiLayerImageTask)), OutputData);
 
 #if WITH_EDITOR
diff --git a/Engine/Plugins/MovieScene/MovieRenderPipeline/Source/MovieRenderPipelineRenderPasses/Private/MoviePipelineImageSequenceOutput.cpp b/Engine/Plugins/MovieScene/MovieRenderPipeline/Source/MovieRenderPipelineRenderPasses/Private/MoviePipelineImageSequenceOutput.cpp
--- a/Engine/Plugins/MovieScene/MovieRenderPipeline/Source/MovieRenderPipelineRenderPasses/Private/MoviePipelineImageSequenceOutput.cpp
+++ b/Engine/Plugins/MovieScene/MovieRenderPipeline/Source/MovieRenderPipelineRenderPasses/Private/MoviePipelineImageSequenceOutput.cpp
@@ -147,11 +147,10 @@ void UMoviePipelineImageSequenceOutputBase::OnReceiveImageDataImpl(FMoviePipelin
 		case EImageFormat::EXR: Extension = TEXT("exr"); break;
 		}
 
-
-		// We need to resolve the filename format string. We combine the folder and file name into one long string first
 		MoviePipeline::FMoviePipelineOutputFutureData OutputData;
 		OutputData.Shot = GetPipeline()->GetActiveShotList()[Payload->SampleState.OutputState.ShotIndex];
 		OutputData.PassIdentifier = RenderPassData.Key;
+		OutputData.OriginSettingClass = GetClass();
 
 		struct FXMLData
 		{
@@ -161,6 +160,7 @@ void UMoviePipelineImageSequenceOutputBase::OnReceiveImageDataImpl(FMoviePipelin
 		
 		FXMLData XMLData;
 		{
+			// We need to resolve the filename format string. We combine the folder and file name into one long string first.
 			FString FileNameFormatString = OutputDirectory / OutputSettings->FileNameFormat;
 
 			// If we're writing more than one render pass out, we need to ensure the file name has the format string in it so we don't
@@ -190,6 +190,15 @@ void UMoviePipelineImageSequenceOutputBase::OnReceiveImageDataImpl(FMoviePipelin
 				{
 					OutputData.FilePath = FPaths::ConvertRelativePathToFull(OutputData.FilePath);
 				}
+
+				// Re-resolve using the frame-templated format string
+				const FString FrameTemplatedFormatString = MoviePipeline::GetFrameTemplatedFilePath(FileNameFormatString);	
+				GetPipeline()->ResolveFilenameFormatArguments(FrameTemplatedFormatString, FormatOverrides, OutputData.FrameTemplatedFilePath, FinalFormatArgs, &Payload->SampleState.OutputState);
+
+				if (FPaths::IsRelative(OutputData.FrameTemplatedFilePath))
+				{
+					OutputData.FrameTemplatedFilePath = FPaths::ConvertRelativePathToFull(OutputData.FrameTemplatedFilePath);
+				}
 			}
 
 			// More XML resolving. Create a deterministic clipname by removing frame numbers, file extension, and any trailing .'s
diff --git a/Engine/Plugins/MovieScene/MovieRenderPipeline/Source/MovieRenderPipelineRenderPasses/Private/MoviePipelineWaveOutput.cpp b/Engine/Plugins/MovieScene/MovieRenderPipeline/Source/MovieRenderPipelineRenderPasses/Private/MoviePipelineWaveOutput.cpp
--- a/Engine/Plugins/MovieScene/MovieRenderPipeline/Source/MovieRenderPipelineRenderPasses/Private/MoviePipelineWaveOutput.cpp
+++ b/Engine/Plugins/MovieScene/MovieRenderPipeline/Source/MovieRenderPipelineRenderPasses/Private/MoviePipelineWaveOutput.cpp
@@ -155,6 +155,7 @@ void UMoviePipelineWaveOutput::BeginFinalizeImpl()
 		MoviePipeline::FMoviePipelineOutputFutureData OutputData;
 		OutputData.Shot = GetPipeline()->GetActiveShotList()[Segment.ShotIndex];
 		OutputData.PassIdentifier = FMoviePipelinePassIdentifier(TEXT("Audio"));
+		OutputData.OriginSettingClass = GetClass();
 
 		// Do this before we start manipulating the filepath for the audio API
 		OutputData.FilePath = Segment.FilePath;
