MoviePipeline: Updated spatial sampling in the graph so it only affects the last 5 warm-up frames. This is optionally adjustable via the "MovieRenderPipeline.NumWarmUpFramesWithSpatialSampling" cvar.

diff --git a/Engine/Plugins/MovieScene/MoviePipelineMaskRenderPass/Source/MoviePipelineMaskRenderPass/Private/Graph/MovieGraphObjectIdNode.cpp b/Engine/Plugins/MovieScene/MoviePipelineMaskRenderPass/Source/MoviePipelineMaskRenderPass/Private/Graph/MovieGraphObjectIdNode.cpp
--- a/Engine/Plugins/MovieScene/MoviePipelineMaskRenderPass/Source/MoviePipelineMaskRenderPass/Private/Graph/MovieGraphObjectIdNode.cpp
+++ b/Engine/Plugins/MovieScene/MoviePipelineMaskRenderPass/Source/MoviePipelineMaskRenderPass/Private/Graph/MovieGraphObjectIdNode.cpp
@@ -33,6 +33,11 @@ FEngineShowFlags UMovieGraphObjectIdNode::GetShowFlags() const
 	return Flags;
 }
 
+int32 UMovieGraphObjectIdNode::GetNumSpatialSamplesDuringWarmUp(const FMovieGraphTraversalContext& InFrameTraversalContext, const FMovieGraphTimeStepData& InTimeData) const
+{
+	return 0;
+}
+
 void UMovieGraphObjectIdNode::GetFormatResolveArgs(FMovieGraphResolveArgs& OutMergedFormatArgs, const FMovieGraphRenderDataIdentifier& InRenderDataIdentifier) const
 {
 	FString MetadataPrefix = UE::MoviePipeline::GetMetadataPrefixPath(InRenderDataIdentifier);
diff --git a/Engine/Plugins/MovieScene/MoviePipelineMaskRenderPass/Source/MoviePipelineMaskRenderPass/Public/Graph/MovieGraphObjectIdNode.h b/Engine/Plugins/MovieScene/MoviePipelineMaskRenderPass/Source/MoviePipelineMaskRenderPass/Public/Graph/MovieGraphObjectIdNode.h
--- a/Engine/Plugins/MovieScene/MoviePipelineMaskRenderPass/Source/MoviePipelineMaskRenderPass/Public/Graph/MovieGraphObjectIdNode.h
+++ b/Engine/Plugins/MovieScene/MoviePipelineMaskRenderPass/Source/MoviePipelineMaskRenderPass/Public/Graph/MovieGraphObjectIdNode.h
@@ -27,7 +27,7 @@ public:
 	UE_API virtual bool GetAllowsCompositing() const override;
 	UE_API virtual bool GetForceLosslessCompression() const override;
 	virtual int32 GetNumSpatialSamples() const override { return SpatialSampleCount; }
-	virtual int32 GetNumSpatialSamplesDuringWarmUp() const  override { return 0; }
+	UE_API virtual int32 GetNumSpatialSamplesDuringWarmUp(const FMovieGraphTraversalContext& InFrameTraversalContext, const FMovieGraphTimeStepData& InTimeData) const override;
 	UE_API virtual void GetFormatResolveArgs(FMovieGraphResolveArgs& OutMergedFormatArgs, const FMovieGraphRenderDataIdentifier& InRenderDataIdentifier) const override;
 	// The ObjectID pass doesn't support any anti-aliasing so we force it to None so that their spatial samples turn into MRG-based camera jitter.
 	virtual bool GetOverrideAntiAliasing() const override { return true; }
diff --git a/Engine/Plugins/MovieScene/MovieRenderPipeline/Source/MovieRenderPipelineRenderPasses/Private/Graph/Nodes/MovieGraphImagePassBaseNode.cpp b/Engine/Plugins/MovieScene/MovieRenderPipeline/Source/MovieRenderPipelineRenderPasses/Private/Graph/Nodes/MovieGraphImagePassBaseNode.cpp
--- a/Engine/Plugins/MovieScene/MovieRenderPipeline/Source/MovieRenderPipelineRenderPasses/Private/Graph/Nodes/MovieGraphImagePassBaseNode.cpp
+++ b/Engine/Plugins/MovieScene/MovieRenderPipeline/Source/MovieRenderPipelineRenderPasses/Private/Graph/Nodes/MovieGraphImagePassBaseNode.cpp
@@ -1,7 +1,9 @@
 // Copyright Epic Games, Inc. All Rights Reserved.
 
 #include "Graph/Nodes/MovieGraphImagePassBaseNode.h"
+
 #include "Graph/MovieGraphRenderLayerSubsystem.h" 
+#include "MoviePipelineQueue.h"
 
 #include UE_INLINE_GENERATED_CPP_BY_NAME(MovieGraphImagePassBaseNode)
 
@@ -16,6 +18,12 @@ static TAutoConsoleVariable<bool> CVarMoviePipelineFlushRenderingPerLayerHack(
 	TEXT(" 1 - Flush (but negatively impacts performance)"),
 	ECVF_Default);
 
+static TAutoConsoleVariable<int32> CVarMoviePipelineNumWarmUpFramesWithSpatialSampling(
+	TEXT("MovieRenderPipeline.NumWarmUpFramesWithSpatialSampling"),
+	5,
+	TEXT("The number of warm-up frames that should have spatial sampling applied to them.\n"),
+	ECVF_Default);
+
 FString UMovieGraphImagePassBaseNode::DefaultDepthAsset = TEXT("/MovieRenderPipeline/Materials/MovieRenderQueue_WorldDepth.MovieRenderQueue_WorldDepth");
 FString UMovieGraphImagePassBaseNode::DefaultMotionVectorsAsset = TEXT("/MovieRenderPipeline/Materials/MovieRenderQueue_MotionVectors.MovieRenderQueue_MotionVectors");
 
@@ -116,6 +124,19 @@ TArray<uint32> UMovieGraphImagePassBaseNode::GetDefaultShowFlags() const
 	return TArray<uint32>();
 }
 
+int32 UMovieGraphImagePassBaseNode::GetNumSpatialSamplesDuringWarmUp(const FMovieGraphTraversalContext& InFrameTraversalContext, const FMovieGraphTimeStepData& InTimeData) const
+{
+	if (InFrameTraversalContext.Shot)
+	{
+		if (InFrameTraversalContext.Shot->ShotInfo.NumEngineWarmUpFramesRemaining > CVarMoviePipelineNumWarmUpFramesWithSpatialSampling.GetValueOnGameThread())
+		{
+			return 1;
+		}
+	}
+	
+	return GetNumSpatialSamples();
+}
+
 void UMovieGraphImagePassBaseNode::AddReferencedObjects(UObject* InThis, FReferenceCollector& Collector)
 {
 	Super::AddReferencedObjects(InThis, Collector);
diff --git a/Engine/Plugins/MovieScene/MovieRenderPipeline/Source/MovieRenderPipelineRenderPasses/Private/Graph/Nodes/MovieGraphPathTracerPassNode.cpp b/Engine/Plugins/MovieScene/MovieRenderPipeline/Source/MovieRenderPipelineRenderPasses/Private/Graph/Nodes/MovieGraphPathTracerPassNode.cpp
--- a/Engine/Plugins/MovieScene/MovieRenderPipeline/Source/MovieRenderPipelineRenderPasses/Private/Graph/Nodes/MovieGraphPathTracerPassNode.cpp
+++ b/Engine/Plugins/MovieScene/MovieRenderPipeline/Source/MovieRenderPipelineRenderPasses/Private/Graph/Nodes/MovieGraphPathTracerPassNode.cpp
@@ -226,14 +226,13 @@ int32 UMovieGraphPathTracerRenderPassNode::GetSeedOffset() const
 	return SeedOffset;
 }
 
-int32 UMovieGraphPathTracerRenderPassNode::GetNumSpatialSamplesDuringWarmUp() const
+int32 UMovieGraphPathTracerRenderPassNode::GetNumSpatialSamplesDuringWarmUp(const FMovieGraphTraversalContext& InFrameTraversalContext, const FMovieGraphTimeStepData& InTimeData) const
 {
 	// Path Tracer doesn't have an image history like the deferred renderer, so it doesn't need
 	// to run all the spatial samples.
 	return 1;
 }
 
-
 bool UMovieGraphPathTracerRenderPassNode::GetDisableToneCurve() const
 {
 	return bDisableToneCurve;
diff --git a/Engine/Plugins/MovieScene/MovieRenderPipeline/Source/MovieRenderPipelineRenderPasses/Private/Graph/Renderers/MovieGraphDeferredPanoramicPass.cpp b/Engine/Plugins/MovieScene/MovieRenderPipeline/Source/MovieRenderPipelineRenderPasses/Private/Graph/Renderers/MovieGraphDeferredPanoramicPass.cpp
--- a/Engine/Plugins/MovieScene/MovieRenderPipeline/Source/MovieRenderPipelineRenderPasses/Private/Graph/Renderers/MovieGraphDeferredPanoramicPass.cpp
+++ b/Engine/Plugins/MovieScene/MovieRenderPipeline/Source/MovieRenderPipelineRenderPasses/Private/Graph/Renderers/MovieGraphDeferredPanoramicPass.cpp
@@ -472,7 +472,7 @@ void FMovieGraphDeferredPanoramicPass::Render(const FMovieGraphTraversalContext&
 		InFrameTraversalContext.Shot->ShotInfo.State == EMovieRenderShotState::Rendering ||
 		InFrameTraversalContext.Shot->ShotInfo.State == EMovieRenderShotState::CoolingDown;
 
-	int32 NumSpatialSamples = FMath::Max(1, bIsRenderingState ? ParentNodeThisFrame->GetNumSpatialSamples() : ParentNodeThisFrame->GetNumSpatialSamplesDuringWarmUp());
+	int32 NumSpatialSamples = FMath::Max(1, bIsRenderingState ? ParentNodeThisFrame->GetNumSpatialSamples() : ParentNodeThisFrame->GetNumSpatialSamplesDuringWarmUp(InFrameTraversalContext, InTimeData));
 	const ESceneCaptureSource SceneCaptureSource = ParentNodeThisFrame->GetDisableToneCurve() ? ESceneCaptureSource::SCS_FinalColorHDR : ESceneCaptureSource::SCS_FinalToneCurveHDR;
 	const EAntiAliasingMethod AntiAliasingMethod = UE::MovieRenderPipeline::GetEffectiveAntiAliasingMethod(ParentNodeThisFrame->GetOverrideAntiAliasing(), ParentNodeThisFrame->GetAntiAliasingMethod());
 
diff --git a/Engine/Plugins/MovieScene/MovieRenderPipeline/Source/MovieRenderPipelineRenderPasses/Private/Graph/Renderers/MovieGraphDeferredPass.cpp b/Engine/Plugins/MovieScene/MovieRenderPipeline/Source/MovieRenderPipelineRenderPasses/Private/Graph/Renderers/MovieGraphDeferredPass.cpp
--- a/Engine/Plugins/MovieScene/MovieRenderPipeline/Source/MovieRenderPipelineRenderPasses/Private/Graph/Renderers/MovieGraphDeferredPass.cpp
+++ b/Engine/Plugins/MovieScene/MovieRenderPipeline/Source/MovieRenderPipelineRenderPasses/Private/Graph/Renderers/MovieGraphDeferredPass.cpp
@@ -266,7 +266,7 @@ void FMovieGraphDeferredPass::Render(const FMovieGraphTraversalContext& InFrameT
 	const bool bWriteAllSamples = ParentNodeThisFrame->GetWriteAllSamples() && InFrameTraversalContext.Shot->ShotInfo.State == EMovieRenderShotState::Rendering;
 	const bool bIsRenderingState = InFrameTraversalContext.Shot->ShotInfo.State == EMovieRenderShotState::Rendering ||
 									InFrameTraversalContext.Shot->ShotInfo.State == EMovieRenderShotState::CoolingDown;
-	int32 NumSpatialSamples = FMath::Max(1, bIsRenderingState ? ParentNodeThisFrame->GetNumSpatialSamples() : ParentNodeThisFrame->GetNumSpatialSamplesDuringWarmUp());
+	int32 NumSpatialSamples = FMath::Max(1, bIsRenderingState ? ParentNodeThisFrame->GetNumSpatialSamples() : ParentNodeThisFrame->GetNumSpatialSamplesDuringWarmUp(InFrameTraversalContext, InTimeData));
 	const int32 SeedOffset = ParentNodeThisFrame->GetSeedOffset();
 	const FIntPoint TileCount = ParentNodeThisFrame->GetTileCount();
 	const float TileOverlapFraction = ParentNodeThisFrame->GetTileOverlapPercentage() / 100.f; // User deals with overlap as 0-100, we want it as 0-1
diff --git a/Engine/Plugins/MovieScene/MovieRenderPipeline/Source/MovieRenderPipelineRenderPasses/Public/Graph/Nodes/MovieGraphImagePassBaseNode.h b/Engine/Plugins/MovieScene/MovieRenderPipeline/Source/MovieRenderPipelineRenderPasses/Public/Graph/Nodes/MovieGraphImagePassBaseNode.h
--- a/Engine/Plugins/MovieScene/MovieRenderPipeline/Source/MovieRenderPipelineRenderPasses/Public/Graph/Nodes/MovieGraphImagePassBaseNode.h
+++ b/Engine/Plugins/MovieScene/MovieRenderPipeline/Source/MovieRenderPipelineRenderPasses/Public/Graph/Nodes/MovieGraphImagePassBaseNode.h
@@ -55,9 +55,16 @@ public:
 	/** How many spatial samples should be rendered each frame? */
 	virtual int32 GetNumSpatialSamples() const { return 1; }
 
-	/** How many spatial samples should be used during warm-up frames? */
+	/** DEPRECATED - How many spatial samples should be used during warm-up frames? */
+	UE_DEPRECATED(5.8, "Override the version that takes arguments instead, this version is no longer called.")
 	virtual int32 GetNumSpatialSamplesDuringWarmUp() const { return GetNumSpatialSamples(); }
 
+	/**
+	 * The number of spatial samples that are performed during warm-up frames. This is called per-frame, so the value returned here can vary
+	 * during the warm-up process.
+	 */
+	UE_API virtual int32 GetNumSpatialSamplesDuringWarmUp(const FMovieGraphTraversalContext& InFrameTraversalContext, const FMovieGraphTimeStepData& InTimeData) const;
+
 	/** Offset to apply to random number generator seed. */
 	virtual int32 GetSeedOffset() const { return 0; }
 
diff --git a/Engine/Plugins/MovieScene/MovieRenderPipeline/Source/MovieRenderPipelineRenderPasses/Public/Graph/Nodes/MovieGraphPathTracerPassNode.h b/Engine/Plugins/MovieScene/MovieRenderPipeline/Source/MovieRenderPipelineRenderPasses/Public/Graph/Nodes/MovieGraphPathTracerPassNode.h
--- a/Engine/Plugins/MovieScene/MovieRenderPipeline/Source/MovieRenderPipelineRenderPasses/Public/Graph/Nodes/MovieGraphPathTracerPassNode.h
+++ b/Engine/Plugins/MovieScene/MovieRenderPipeline/Source/MovieRenderPipelineRenderPasses/Public/Graph/Nodes/MovieGraphPathTracerPassNode.h
@@ -52,7 +52,7 @@ public:
 	UE_API virtual TArray<FMoviePipelinePostProcessPass> GetAdditionalPostProcessMaterials() const override;
 	UE_API virtual FString GetPPMFileNameFormatString() const override;
 	UE_API virtual int32 GetNumSpatialSamples() const override;
-	UE_API virtual int32 GetNumSpatialSamplesDuringWarmUp() const override;
+	UE_API virtual int32 GetNumSpatialSamplesDuringWarmUp(const FMovieGraphTraversalContext& InFrameTraversalContext, const FMovieGraphTimeStepData& InTimeData) const override;
 	UE_API virtual int32 GetSeedOffset() const override;
 	UE_API virtual bool GetDisableToneCurve() const override;
 	UE_API virtual bool GetAllowOCIO() const override;
