MoviePipeline: Updated the autocomplete used in MRG to the one used in CAT in order to support MRG tokens as well as CAT naming tokens.

Additionally, all MRG node properties that support {tokens} now display the autocomplete widget.

diff --git a/Engine/Plugins/Media/AppleProResMedia/Source/AppleProResMedia/Private/MovieGraphAppleProResNode.h b/Engine/Plugins/Media/AppleProResMedia/Source/AppleProResMedia/Private/MovieGraphAppleProResNode.h
--- a/Engine/Plugins/Media/AppleProResMedia/Source/AppleProResMedia/Private/MovieGraphAppleProResNode.h
+++ b/Engine/Plugins/Media/AppleProResMedia/Source/AppleProResMedia/Private/MovieGraphAppleProResNode.h
@@ -121,6 +121,6 @@ public:
 	bool bCompositeOntoFinalImage = true;
 
 	/** The file name format used for writing out the burn-in video (if not composited). */
-	UPROPERTY(EditAnywhere, BlueprintReadWrite, Category = "Burn In", meta = (EditCondition = "bOverride_BurnInFileNameFormat"))
+	UPROPERTY(EditAnywhere, BlueprintReadWrite, Category = "Burn In", meta = (EditCondition = "bOverride_BurnInFileNameFormat", MrgTokenAutocomplete))
 	FString BurnInFileNameFormat = TEXT("{sequence_name}.{layer_name}.{renderer_name}");
 };
\ No newline at end of file
diff --git a/Engine/Plugins/Media/AvidDNxMedia/Source/Source/Private/MovieGraphAvidDNxHRNode.h b/Engine/Plugins/Media/AvidDNxMedia/Source/Source/Private/MovieGraphAvidDNxHRNode.h
--- a/Engine/Plugins/Media/AvidDNxMedia/Source/Source/Private/MovieGraphAvidDNxHRNode.h
+++ b/Engine/Plugins/Media/AvidDNxMedia/Source/Source/Private/MovieGraphAvidDNxHRNode.h
@@ -115,6 +115,6 @@ public:
 	bool bCompositeOntoFinalImage = true;
 
 	/** The file name format used for writing out the burn-in video (if not composited). */
-	UPROPERTY(EditAnywhere, BlueprintReadWrite, Category = "Burn In", meta = (EditCondition = "bOverride_BurnInFileNameFormat"))
+	UPROPERTY(EditAnywhere, BlueprintReadWrite, Category = "Burn In", meta = (EditCondition = "bOverride_BurnInFileNameFormat", MrgTokenAutocomplete))
 	FString BurnInFileNameFormat = TEXT("{sequence_name}.{layer_name}.{renderer_name}");
 };
\ No newline at end of file
diff --git a/Engine/Plugins/MovieScene/MovieRenderPipeline/Source/MovieRenderPipelineCore/Private/Graph/MovieGraphBlueprintLibrary.cpp b/Engine/Plugins/MovieScene/MovieRenderPipeline/Source/MovieRenderPipelineCore/Private/Graph/MovieGraphBlueprintLibrary.cpp
--- a/Engine/Plugins/MovieScene/MovieRenderPipeline/Source/MovieRenderPipelineCore/Private/Graph/MovieGraphBlueprintLibrary.cpp
+++ b/Engine/Plugins/MovieScene/MovieRenderPipeline/Source/MovieRenderPipelineCore/Private/Graph/MovieGraphBlueprintLibrary.cpp
@@ -146,6 +146,24 @@ namespace
 		}
 	}
 
+	/** Removes all "mrg:" namespacing from {tokens} in the format string. For example, replaces {mrg:branch_name} with {branch_name}. */
+	void RemoveMrgTokenNamespacing(FString& InOutFormatString)
+	{
+		static const FString MrgNamespace = FString::Printf(TEXT("%s:"), *UMovieRenderPipelineNamingTokens::TokenNamespace);
+		
+		for (const FString& IncludedTokenName : GetTokensFromString(InOutFormatString))
+		{
+			if (IncludedTokenName.StartsWith(MrgNamespace))
+			{
+				const FString TokenNameWithoutNamespace = IncludedTokenName.Replace(*MrgNamespace, TEXT(""));
+				const FString SourceToken = FString::Printf(TEXT("{%s}"), *IncludedTokenName);
+				const FString DestToken = FString::Printf(TEXT("{%s}"), *TokenNameWithoutNamespace);
+				
+				InOutFormatString.ReplaceInline(*SourceToken, *DestToken, ESearchCase::CaseSensitive);
+			}
+		}
+	}
+
 	FString ResolveFormatArgumentsHelper(const FString& InFormatString, const FMovieGraphFilenameResolveParams& InParams, const bool bResolveToPath, FMovieGraphResolveArgs& OutMergedFormatArgs)
 	{
 		// There are a number of different sources for file/metadata KVPs. 
@@ -305,6 +323,10 @@ namespace
 		// The format string may need to be modified from this point onwards, so make a copy so the source string isn't unintentionally changed.
 		FString FormatStringCopy = InFormatString;
 
+		// The "mrg:" namespace is technically allowed. However, for the purposes of resolving, we need to strip off this namespace because the token
+		// names discovered above do not include this namespace, and the Naming Tokens system will not know how to resolve these either.
+		RemoveMrgTokenNamespacing(FormatStringCopy);
+
 		// Resolve non-MRG tokens via the Naming Tokens system.
 		UpdateTokenMapsWithResolvedExternalTokens(FormatStringCopy, GeneratedResolveArgs.FilenameArguments, NamedArgs);
 
diff --git a/Engine/Plugins/MovieScene/MovieRenderPipeline/Source/MovieRenderPipelineCore/Private/MovieRenderPipelineCoreModule.cpp b/Engine/Plugins/MovieScene/MovieRenderPipeline/Source/MovieRenderPipelineCore/Private/MovieRenderPipelineCoreModule.cpp
--- a/Engine/Plugins/MovieScene/MovieRenderPipeline/Source/MovieRenderPipelineCore/Private/MovieRenderPipelineCoreModule.cpp
+++ b/Engine/Plugins/MovieScene/MovieRenderPipeline/Source/MovieRenderPipelineCore/Private/MovieRenderPipelineCoreModule.cpp
@@ -5,14 +5,22 @@
 #include "Misc/CoreDelegates.h"
 #include "Modules/ModuleManager.h"
 #include "MovieRenderPipelineDataTypes.h"
+#include "Engine/Engine.h"
 #include "Engine/World.h"
 #include "HAL/IConsoleManager.h"
 #include "TimerManager.h"
 #include "MoviePipeline.h"
+#include "Graph/MovieGraphBlueprintLibrary.h"
 #include "Graph/MovieGraphPipeline.h"
 #include "Graph/Nodes/MovieGraphBurnInNode.h"
+#include "NamingTokenData.h"
+#include "NamingTokens.h"
+#include "NamingTokensEngineSubsystem.h"
 #include "UObject/ICookInfo.h"
 
+FGuid FMovieRenderPipelineCoreModule::NamingTokensGuid = FGuid();
+TWeakObjectPtr<UNamingTokens> FMovieRenderPipelineCoreModule::NamingTokens = nullptr;
+
 #if WITH_EDITOR
 static TAutoConsoleVariable<bool> CVarAddCodeReferencedAssetsToCook(
 	TEXT("MoviePipeline.AddCodeReferencedAssetsToCook"), true,
@@ -63,6 +71,8 @@ void FMovieRenderPipelineCoreModule::StartupModule()
 		// Register a hook to wait until the engine has finished loading to increase the likelihood that the desired classes are loaded.
 		FCoreUObjectDelegates::PostLoadMapWithWorld.AddRaw(this, &FMovieRenderPipelineCoreModule::OnMapLoadFinished);
 	}
+
+	FCoreDelegates::OnPostEngineInit.AddRaw(this, &FMovieRenderPipelineCoreModule::BindToNamingTokenEvents);
 }
 
 void FMovieRenderPipelineCoreModule::OnMapLoadFinished(UWorld* InWorld)
@@ -86,6 +96,13 @@ void FMovieRenderPipelineCoreModule::QueueInitialize(UWorld* InWorld)
 
 void FMovieRenderPipelineCoreModule::ShutdownModule()
 {
+	FCoreDelegates::OnPostEngineInit.RemoveAll(this);
+
+	if (const TStrongObjectPtr<UNamingTokens> NamingTokensPin = NamingTokens.Pin())
+	{
+		NamingTokensPin->UnregisterExternalTokens(NamingTokensGuid);
+		NamingTokensPin->GetOnPreEvaluateEvent().RemoveAll(this);
+	}
 }
 
 void FMovieRenderPipelineCoreModule::SetTickInfo(const FMoviePipelineLightweightTickInfo& InTickInfo)
@@ -94,6 +111,52 @@ void FMovieRenderPipelineCoreModule::SetTickInfo(const FMoviePipelineLightweight
 	MRQModule.TickInfo = InTickInfo;
 }
 
+void FMovieRenderPipelineCoreModule::BindToNamingTokenEvents()
+{
+	if (GEngine && !NamingTokens.IsValid())
+	{
+		NamingTokens = GEngine->GetEngineSubsystem<UNamingTokensEngineSubsystem>()->GetNamingTokens(UMovieRenderPipelineNamingTokens::TokenNamespace);
+		if (ensure(NamingTokens.IsValid()))
+		{
+			const TStrongObjectPtr<UNamingTokens> NamingTokensPin = NamingTokens.Pin();
+			NamingTokensPin->RegisterExternalTokens(NamingTokensGuid);
+			NamingTokensPin->GetOnPreEvaluateEvent().AddRaw(this, &FMovieRenderPipelineCoreModule::OnMovieRenderPipelineNamingTokensPreEvaluate);
+		}
+	}
+}
+
+void FMovieRenderPipelineCoreModule::OnMovieRenderPipelineNamingTokensPreEvaluate(const FNamingTokensEvaluationData& InEvaluationData)
+{
+	if (NamingTokens.IsValid())
+	{
+		TArray<FNamingTokenData>& ExternalTokens = NamingTokens->GetExternalTokensChecked(NamingTokensGuid);
+
+		// Only add tokens if they have not yet been registered with the Naming Tokens system
+		if (ExternalTokens.IsEmpty())
+		{
+			// Just fetch the format arguments, no need to specify a format string; all we care about here is the token list
+			FMovieGraphResolveArgs FormatArgs;
+			UMovieGraphBlueprintLibrary::ResolveFormatArguments(FString(), FMovieGraphFilenameResolveParams(), FormatArgs);
+
+			// Sort the format tokens alphabetically
+			TArray<FString> FormatTokens;
+			FormatArgs.FilenameArguments.GetKeys(FormatTokens);
+			FormatTokens.Sort();
+
+			// Add all of the tokens. Note that they do not evaluate to a value (empty delegate) because they're evaluated by MRG, not
+			// the Naming Tokens system.
+			for (const FString& FormatToken : FormatTokens)
+			{
+				ExternalTokens.Add({
+					FormatToken,
+					FText::FromString(FormatToken),
+					FNamingTokenData::FTokenProcessorDelegateNative::CreateLambda([]() { return FText::GetEmpty(); })
+				});
+			}
+		}
+	}
+}
+
 IMPLEMENT_MODULE(FMovieRenderPipelineCoreModule, MovieRenderPipelineCore);
 DEFINE_LOG_CATEGORY(LogMovieRenderPipeline); 
 DEFINE_LOG_CATEGORY(LogMovieRenderPipelineIO);
diff --git a/Engine/Plugins/MovieScene/MovieRenderPipeline/Source/MovieRenderPipelineCore/Public/Graph/MovieGraphNode.h b/Engine/Plugins/MovieScene/MovieRenderPipeline/Source/MovieRenderPipelineCore/Public/Graph/MovieGraphNode.h
--- a/Engine/Plugins/MovieScene/MovieRenderPipeline/Source/MovieRenderPipelineCore/Public/Graph/MovieGraphNode.h
+++ b/Engine/Plugins/MovieScene/MovieRenderPipeline/Source/MovieRenderPipelineCore/Public/Graph/MovieGraphNode.h
@@ -530,6 +530,10 @@ public:
 	 * graph evaluation lifecycle that this resolve step occurs (eg, {camera_name}).
 	 *
 	 * Nodes that implement this method should also take care to call the super class in order to ensure inherited properties are covered.
+	 *
+	 * If the properties resolved here should use the token autocomplete widget in the details panel, include the MrgTokenAutocomplete
+	 * UPROPERTY metadata on them (eg, `meta=(MrgTokenAutocomplete)`). Note that this metadata will only take effect for top-level properties
+	 * (ie, properties that are directly declared on the node); properties that are embedded in structs, for example, will not be detected.
 	 */
 	MOVIERENDERPIPELINECORE_API virtual void ResolveTokenContainingProperties(TFunction<void(FString&)>& ResolveFunc, const FMovieGraphTokenResolveContext& InContext);
 	
diff --git a/Engine/Plugins/MovieScene/MovieRenderPipeline/Source/MovieRenderPipelineCore/Public/Graph/Nodes/MovieGraphCommandLineEncoderNode.h b/Engine/Plugins/MovieScene/MovieRenderPipeline/Source/MovieRenderPipelineCore/Public/Graph/Nodes/MovieGraphCommandLineEncoderNode.h
--- a/Engine/Plugins/MovieScene/MovieRenderPipeline/Source/MovieRenderPipelineCore/Public/Graph/Nodes/MovieGraphCommandLineEncoderNode.h
+++ b/Engine/Plugins/MovieScene/MovieRenderPipeline/Source/MovieRenderPipelineCore/Public/Graph/Nodes/MovieGraphCommandLineEncoderNode.h
@@ -82,7 +82,7 @@ public:
 	* Can be different from the main one in the Output setting so you can render out frames to individual
 	* shot folders but encode to one file.
 	*/
-	UPROPERTY(EditAnywhere, BlueprintReadWrite, Category = "General", meta=(EditCondition="bOverride_FileNameFormat"))
+	UPROPERTY(EditAnywhere, BlueprintReadWrite, Category = "General", meta=(EditCondition="bOverride_FileNameFormat", MrgTokenAutocomplete))
 	FString FileNameFormat;
 
 	/** Whether the source files should be deleted on disk after encoding. */
diff --git a/Engine/Plugins/MovieScene/MovieRenderPipeline/Source/MovieRenderPipelineCore/Public/Graph/Nodes/MovieGraphDebugNode.h b/Engine/Plugins/MovieScene/MovieRenderPipeline/Source/MovieRenderPipelineCore/Public/Graph/Nodes/MovieGraphDebugNode.h
--- a/Engine/Plugins/MovieScene/MovieRenderPipeline/Source/MovieRenderPipelineCore/Public/Graph/Nodes/MovieGraphDebugNode.h
+++ b/Engine/Plugins/MovieScene/MovieRenderPipeline/Source/MovieRenderPipelineCore/Public/Graph/Nodes/MovieGraphDebugNode.h
@@ -79,7 +79,7 @@ public:
 	* If bCaptureUnrealInsightsTrace is true, name of the UnrealInsights trace.
 	* Resulting capture will be in the global Output Directory for the job.
 	*/
-	UPROPERTY(EditAnywhere, BlueprintReadWrite, Category = "Unreal Insights", meta = (EditCondition = "bOverride_UnrealInsightsTraceFileNameFormat"))
+	UPROPERTY(EditAnywhere, BlueprintReadWrite, Category = "Unreal Insights", meta = (EditCondition = "bOverride_UnrealInsightsTraceFileNameFormat", MrgTokenAutocomplete))
 	FString UnrealInsightsTraceFileNameFormat = TEXT("{sequence_name}_UnrealInsights");
 
 	/** 
@@ -103,7 +103,7 @@ public:
 	/**
 	* If bCaptureFramesWithDumpGpu is true, name of the output directory for the capture (related to frame output directory).
 	*/
-	UPROPERTY(EditAnywhere, BlueprintReadWrite, Category = "DumpGPU", meta = (EditCondition = "bOverride_DumpGpuCaptureFilePathFormat"))
+	UPROPERTY(EditAnywhere, BlueprintReadWrite, Category = "DumpGPU", meta = (EditCondition = "bOverride_DumpGpuCaptureFilePathFormat", MrgTokenAutocomplete))
 	FString DumpGpuCaptureFilePathFormat = TEXT("{sequence_name}.{frame_number}_DumpGPU");
 
 	/**
diff --git a/Engine/Plugins/MovieScene/MovieRenderPipeline/Source/MovieRenderPipelineCore/Public/Graph/Nodes/MovieGraphFileOutputNode.h b/Engine/Plugins/MovieScene/MovieRenderPipeline/Source/MovieRenderPipelineCore/Public/Graph/Nodes/MovieGraphFileOutputNode.h
--- a/Engine/Plugins/MovieScene/MovieRenderPipeline/Source/MovieRenderPipelineCore/Public/Graph/Nodes/MovieGraphFileOutputNode.h
+++ b/Engine/Plugins/MovieScene/MovieRenderPipeline/Source/MovieRenderPipelineCore/Public/Graph/Nodes/MovieGraphFileOutputNode.h
@@ -77,7 +77,7 @@ public:
 	uint8 bOverride_FileNameFormat : 1;
 
 	/** What format string should the final files use? Can include folder prefixes, and format string ({shot_name}, etc.) */
-	UPROPERTY(EditAnywhere, BlueprintReadWrite, Category = "File Output", meta=(EditCondition="bOverride_FileNameFormat"))
+	UPROPERTY(EditAnywhere, BlueprintReadWrite, Category = "File Output", meta=(EditCondition="bOverride_FileNameFormat", MrgTokenAutocomplete))
 	FString FileNameFormat = "{sequence_name}.{layer_name}.{frame_number}";
 
 protected:
diff --git a/Engine/Plugins/MovieScene/MovieRenderPipeline/Source/MovieRenderPipelineCore/Public/Graph/Nodes/MovieGraphGlobalOutputSettingNode.h b/Engine/Plugins/MovieScene/MovieRenderPipeline/Source/MovieRenderPipelineCore/Public/Graph/Nodes/MovieGraphGlobalOutputSettingNode.h
--- a/Engine/Plugins/MovieScene/MovieRenderPipeline/Source/MovieRenderPipelineCore/Public/Graph/Nodes/MovieGraphGlobalOutputSettingNode.h
+++ b/Engine/Plugins/MovieScene/MovieRenderPipeline/Source/MovieRenderPipelineCore/Public/Graph/Nodes/MovieGraphGlobalOutputSettingNode.h
@@ -195,7 +195,7 @@ public:
 	uint8 bOverride_bFlushDiskWritesPerShot : 1;
 
 	/** What directory should all of our output files be relative to. */
-	UPROPERTY(EditAnywhere, BlueprintReadWrite, Category = "File Output", meta=(EditCondition="bOverride_OutputDirectory"))
+	UPROPERTY(EditAnywhere, BlueprintReadWrite, Category = "File Output", meta=(EditCondition="bOverride_OutputDirectory", MrgTokenAutocomplete))
 	FDirectoryPath OutputDirectory;
 
 	/** What resolution should our output files be exported at? */
diff --git a/Engine/Plugins/MovieScene/MovieRenderPipeline/Source/MovieRenderPipelineCore/Public/Graph/Nodes/MovieGraphRenderPassNode.h b/Engine/Plugins/MovieScene/MovieRenderPipeline/Source/MovieRenderPipelineCore/Public/Graph/Nodes/MovieGraphRenderPassNode.h
--- a/Engine/Plugins/MovieScene/MovieRenderPipeline/Source/MovieRenderPipelineCore/Public/Graph/Nodes/MovieGraphRenderPassNode.h
+++ b/Engine/Plugins/MovieScene/MovieRenderPipeline/Source/MovieRenderPipelineCore/Public/Graph/Nodes/MovieGraphRenderPassNode.h
@@ -109,11 +109,11 @@ protected:
 	uint8 bOverride_RendererSubName : 1;
 
 	/** The value that will be used in the {renderer_name} token. */
-	UPROPERTY(EditAnywhere, BlueprintReadWrite, Category = "Naming", meta = (EditCondition="bOverride_RendererName"))
+	UPROPERTY(EditAnywhere, BlueprintReadWrite, Category = "Naming", meta = (EditCondition="bOverride_RendererName", MrgTokenAutocomplete))
 	FString RendererName;
 
 	/** The value that will be used in the {renderer_sub_name} token. */
-	UPROPERTY(EditAnywhere, BlueprintReadWrite, Category = "Naming", meta = (EditCondition="bOverride_RendererSubName"))
+	UPROPERTY(EditAnywhere, BlueprintReadWrite, Category = "Naming", meta = (EditCondition="bOverride_RendererSubName", MrgTokenAutocomplete))
 	FString RendererSubName;
 
 protected:
diff --git a/Engine/Plugins/MovieScene/MovieRenderPipeline/Source/MovieRenderPipelineCore/Public/MovieRenderPipelineCoreModule.h b/Engine/Plugins/MovieScene/MovieRenderPipeline/Source/MovieRenderPipelineCore/Public/MovieRenderPipelineCoreModule.h
--- a/Engine/Plugins/MovieScene/MovieRenderPipeline/Source/MovieRenderPipelineCore/Public/MovieRenderPipelineCoreModule.h
+++ b/Engine/Plugins/MovieScene/MovieRenderPipeline/Source/MovieRenderPipelineCore/Public/MovieRenderPipelineCoreModule.h
@@ -3,22 +3,50 @@
 #pragma once
 
 #include "CoreMinimal.h"
+#include "NamingTokens.h"
 #include "Stats/Stats.h"
 #include "Modules/ModuleInterface.h"
+#include "NamingTokensEvaluationData.h"
 #include "Features/IModularFeature.h"
 
+#include "MovieRenderPipelineCoreModule.generated.h"
+
 #define UE_API MOVIERENDERPIPELINECORE_API
 
 // Forward Declare
-class UMoviePipelineExecutorBase;
-class UMoviePipelineQueue;
 class ULevelSequence;
 class UMoviePipeline;
+class UMoviePipelineExecutorBase;
+class UMoviePipelineQueue;
 struct FMoviePipelineFrameOutputState;
 
 // Declare a stat-group for our performance stats to be counted under, readable in game by "stat MovieRenderPipeline".
 DECLARE_STATS_GROUP(TEXT("MovieRenderPipeline"), STATGROUP_MoviePipeline, STATCAT_Advanced);
 
+/**
+ * Naming Tokens for Movie Render Pipeline.
+ */
+UCLASS(NotBlueprintable)
+class UMovieRenderPipelineNamingTokens final: public UNamingTokens
+{
+	GENERATED_BODY()
+
+public:
+	/** The namespace that all MRG tokens are included under within the Naming Tokens system. */
+	static inline const FString TokenNamespace = TEXT("mrg");
+	
+	UMovieRenderPipelineNamingTokens()
+	{
+		Namespace = TokenNamespace;
+		NamespaceDisplayName = FText::FromString(TEXT("Movie Render Graph"));
+		bIsPrivateNamespace = true;
+	}
+	
+	virtual ~UMovieRenderPipelineNamingTokens() override
+	{
+	}
+};
+
 /**
 * Movie pipeline tick information
 * Used by external consumers, like Niagara, to query information
@@ -71,8 +99,20 @@ private:
 
 	UE_API uint8 ParseMovieRenderData(const FString& InSequenceAssetPath, const FString& InConfigAssetPath, const FString& InExecutorType, const FString& InPipelineType,
 		UMoviePipelineQueue*& OutQueue, UMoviePipelineExecutorBase*& OutExecutor) const;
+	
+	/** Setup bindings to MRQ Naming Tokens. */
+	void BindToNamingTokenEvents();
+
+	/** Executes every time Naming Tokens are fetched. */
+	void OnMovieRenderPipelineNamingTokensPreEvaluate(const FNamingTokensEvaluationData& InEvaluationData);
 
 private:
+	/** The GUID that the Naming Tokens system has assigned to the MRQ naming tokens. */
+	static FGuid NamingTokensGuid;
+
+	/** Naming Tokens object for MRQ. */
+	static TWeakObjectPtr<UNamingTokens> NamingTokens;
+	
 	FMoviePipelineLightweightTickInfo TickInfo;
 	FString MoviePipelineLocalExecutorClassType;
 	FString MoviePipelineClassType;
diff --git a/Engine/Plugins/MovieScene/MovieRenderPipeline/Source/MovieRenderPipelineEditor/MovieRenderPipelineEditor.Build.cs b/Engine/Plugins/MovieScene/MovieRenderPipeline/Source/MovieRenderPipelineEditor/MovieRenderPipelineEditor.Build.cs
--- a/Engine/Plugins/MovieScene/MovieRenderPipeline/Source/MovieRenderPipelineEditor/MovieRenderPipelineEditor.Build.cs
+++ b/Engine/Plugins/MovieScene/MovieRenderPipeline/Source/MovieRenderPipelineEditor/MovieRenderPipelineEditor.Build.cs
@@ -76,6 +76,7 @@ public class MovieRenderPipelineEditor : ModuleRules
 				"Imath",		// This and UEOpenExr are only required for use of UMovieGraphImageSequenceOutputNode. Ideally these are refactored
 				"UEOpenExr",	// out of the ImageSequenceOutputNode header at some point so these dependencies can be removed.
 				"Projects",		// For IPluginManager
+				"NamingTokensUI",
 			}
         );
 
diff --git a/Engine/Plugins/MovieScene/MovieRenderPipeline/Source/MovieRenderPipelineEditor/Private/Customizations/Graph/MovieGraphCustomizationUtils.cpp b/Engine/Plugins/MovieScene/MovieRenderPipeline/Source/MovieRenderPipelineEditor/Private/Customizations/Graph/MovieGraphCustomizationUtils.cpp
--- a/Engine/Plugins/MovieScene/MovieRenderPipeline/Source/MovieRenderPipelineEditor/Private/Customizations/Graph/MovieGraphCustomizationUtils.cpp
+++ b/Engine/Plugins/MovieScene/MovieRenderPipeline/Source/MovieRenderPipelineEditor/Private/Customizations/Graph/MovieGraphCustomizationUtils.cpp
@@ -4,9 +4,14 @@
 
 #include "DetailCategoryBuilder.h"
 #include "DetailLayoutBuilder.h"
+#include "DetailWidgetRow.h"
 #include "Graph/MovieGraphConfig.h"
 #include "IDetailGroup.h"
 #include "MovieJobVariableAssignmentContainer.h"
+#include "MovieRenderPipelineCoreModule.h"
+#include "MovieRenderPipelineStyle.h"
+#include "NamingTokensEngineSubsystem.h"
+#include "SNamingTokensEditableTextBox.h"
 
 namespace UE::MovieRenderPipelineEditor::Private
 {
@@ -38,4 +43,94 @@ namespace UE::MovieRenderPipelineEditor::Private
 			}
 		}
 	}
+
+	void AddTokenAutocompleteToPropertyRow(IDetailPropertyRow* InPropertyRow)
+	{
+		/** Determines if the given property points to a FDirectoryPath struct. */
+		auto IsDirectoryPathProperty = [](const TSharedPtr<IPropertyHandle>& InPropertyHandle)
+		{
+			const FProperty* Property = InPropertyHandle->GetProperty();
+			if (const FStructProperty* StructProperty = CastField<FStructProperty>(Property))
+			{
+				if (const TObjectPtr<UScriptStruct> Struct = StructProperty->Struct)
+				{
+					return Struct->IsChildOf<FDirectoryPath>();
+				}
+			}
+
+			return false;
+		};
+
+		if (!InPropertyRow)
+		{
+			return;
+		}
+
+		// Note: For now, FDirectoryPath support is disabled because it causes the "..." button to disappear. At some point ideally FDirectoryPath
+		// also includes token autocomplete, but the "..." button is more important.
+		TSharedPtr<IPropertyHandle> PropertyHandle = InPropertyRow->GetPropertyHandle();
+		if (IsDirectoryPathProperty(PropertyHandle))
+		{
+			return;
+		}
+
+		// Most token format strings are specified via simple FString properties. However, there are cases where token autocomplete is needed for
+		// FDirectoryPath structs too, so those need to be special-cased here.
+		if (IsDirectoryPathProperty(PropertyHandle))
+		{
+			if (const TSharedPtr<IPropertyHandle> PathProperty = PropertyHandle->GetChildHandle(GET_MEMBER_NAME_CHECKED(FDirectoryPath, Path)))
+			{
+				PropertyHandle = PathProperty;
+			}
+		}
+
+		// The MRG namespace in Naming Tokens is private, so we need to explicitly include the namespace so the MRG tokens show up
+		FNamingTokenFilterArgs NamingTokenFilterArgs;
+		NamingTokenFilterArgs.AdditionalNamespacesToInclude.Add(UMovieRenderPipelineNamingTokens::TokenNamespace);
+
+		InPropertyRow->CustomWidget()
+			.PropertyHandleList({PropertyHandle})
+			.NameContent()
+			[
+				PropertyHandle->CreatePropertyNameWidget()
+			]
+			.ValueContent()
+			.MinDesiredWidth(200.0f)
+			[
+				SNew(SNamingTokensEditableTextBox)
+					.Text_Lambda([WeakAutocompleteProperty = PropertyHandle.ToWeakPtr()]()
+					{
+						FString InitialTextString;
+						
+						if (const TSharedPtr<IPropertyHandle> PropertyPin = WeakAutocompleteProperty.Pin())
+						{
+							PropertyPin->GetValue(InitialTextString);
+						}
+
+						return FText::FromString(InitialTextString);
+					})
+					.OnTextCommitted_Lambda([WeakAutocompleteProperty = PropertyHandle.ToWeakPtr()](const FText& InNewText, ETextCommit::Type CommitType)
+					{
+						if (const TSharedPtr<IPropertyHandle> PropertyPin = WeakAutocompleteProperty.Pin())
+						{
+							const FPropertyAccess::Result Result = PropertyPin->SetValueFromFormattedString(InNewText.ToString());
+							
+							if (Result != FPropertyAccess::Success)
+							{
+								UE_LOG(LogMovieRenderPipeline, Warning, TEXT("Unable to insert token into format string (could not update FProperty)."));
+							}
+						}
+					})
+					.Contexts({})
+					.FilterArgs(NamingTokenFilterArgs)
+					.NamespaceSuggestionPriority({ UMovieRenderPipelineNamingTokens::TokenNamespace })
+					.IsReadOnly(false)
+					.CanDisplayResolvedText(false)
+					.AllowMultiLine(false)
+					.EnableSuggestionDropdown(true)
+					.TextBoxPadding(FMargin(4))
+					.Style(&FMovieRenderPipelineStyle::Get().GetWidgetStyle<FEditableTextBoxStyle>("MovieRenderPipeline.AutocompleteTextStyle"))
+					.ArgumentStyle(&FMovieRenderPipelineStyle::Get().GetWidgetStyle<FTextBlockStyle>("MovieRenderPipeline.AutocompleteArgumentTextStyle"))
+			];
+	}
 }
diff --git a/Engine/Plugins/MovieScene/MovieRenderPipeline/Source/MovieRenderPipelineEditor/Private/Customizations/Graph/MovieGraphCustomizationUtils.h b/Engine/Plugins/MovieScene/MovieRenderPipeline/Source/MovieRenderPipelineEditor/Private/Customizations/Graph/MovieGraphCustomizationUtils.h
--- a/Engine/Plugins/MovieScene/MovieRenderPipeline/Source/MovieRenderPipelineEditor/Private/Customizations/Graph/MovieGraphCustomizationUtils.h
+++ b/Engine/Plugins/MovieScene/MovieRenderPipeline/Source/MovieRenderPipelineEditor/Private/Customizations/Graph/MovieGraphCustomizationUtils.h
@@ -14,4 +14,7 @@ namespace UE::MovieRenderPipelineEditor::Private
 	 * the name of the graph) is added under the given category.
 	 */
 	void AddVariableAssignments(TArray<TObjectPtr<UMovieJobVariableAssignmentContainer>>& InVariableAssignments, IDetailCategoryBuilder& InCategory, IDetailLayoutBuilder* InDetailBuilder);
+
+	/** Adds the Naming Tokens token autocomplete widget to the given property row. */
+	void AddTokenAutocompleteToPropertyRow(IDetailPropertyRow* InPropertyRow);
 }
diff --git a/Engine/Plugins/MovieScene/MovieRenderPipeline/Source/MovieRenderPipelineEditor/Private/Customizations/Graph/MovieGraphFormatTokenCustomization.cpp b/Engine/Plugins/MovieScene/MovieRenderPipeline/Source/MovieRenderPipelineEditor/Private/Customizations/Graph/MovieGraphFormatTokenCustomization.cpp
deleted file mode 100644
diff --git a/Engine/Plugins/MovieScene/MovieRenderPipeline/Source/MovieRenderPipelineEditor/Private/Customizations/Graph/MovieGraphFormatTokenCustomization.h b/Engine/Plugins/MovieScene/MovieRenderPipeline/Source/MovieRenderPipelineEditor/Private/Customizations/Graph/MovieGraphFormatTokenCustomization.h
deleted file mode 100644
diff --git a/Engine/Plugins/MovieScene/MovieRenderPipeline/Source/MovieRenderPipelineEditor/Private/Customizations/Graph/MovieGraphNodeCustomization.h b/Engine/Plugins/MovieScene/MovieRenderPipeline/Source/MovieRenderPipelineEditor/Private/Customizations/Graph/MovieGraphNodeCustomization.h
--- a/Engine/Plugins/MovieScene/MovieRenderPipeline/Source/MovieRenderPipelineEditor/Private/Customizations/Graph/MovieGraphNodeCustomization.h
+++ b/Engine/Plugins/MovieScene/MovieRenderPipeline/Source/MovieRenderPipelineEditor/Private/Customizations/Graph/MovieGraphNodeCustomization.h
@@ -4,8 +4,9 @@
 
 #include "DetailCategoryBuilder.h"
 #include "DetailLayoutBuilder.h"
+#include "Graph/MovieGraphNode.h"
 #include "IDetailCustomization.h"
-#include "Graph/MovieGraphConfig.h"
+#include "MovieGraphCustomizationUtils.h"
 #include "PropertyHandle.h"
 
 #define LOCTEXT_NAMESPACE "MoviePipelineEditor"
@@ -28,13 +29,14 @@ protected:
 
 	virtual void CustomizeDetails(IDetailLayoutBuilder& DetailBuilder) override
 	{
-		TArray<TWeakObjectPtr<UObject>> ObjectsBeingCustomized;
-		DetailBuilder.GetObjectsBeingCustomized(ObjectsBeingCustomized);
+		AddTokenAutocompleteWidgetToAllEligibleProperties(DetailBuilder);
 
 		// Hide the "Properties" category, which houses the dynamic properties, if there are no dynamic properties
 		// on the node. An empty category in the details panel is messy/confusing.
 		//
 		// Some nodes may also request that specific properties (most likely inherited) are hidden.
+		TArray<TWeakObjectPtr<UObject>> ObjectsBeingCustomized;
+		DetailBuilder.GetObjectsBeingCustomized(ObjectsBeingCustomized);
 		for (const TWeakObjectPtr<UObject>& CustomizedObject : ObjectsBeingCustomized)
 		{
 			if (const UMovieGraphNode* Node = Cast<UMovieGraphNode>(CustomizedObject))
@@ -71,6 +73,43 @@ protected:
 		DetailBuilder.EditCategory("File Output").SetSortOrder(ECategoryPriority::Important);
 	}
 	//~ End IDetailCustomization interface
+
+private:
+	/** Adds the token autocomplete widget to all properties that have the MrgTokenAutocomplete UPROPERTY metadata. */
+	void AddTokenAutocompleteWidgetToAllEligibleProperties(IDetailLayoutBuilder& DetailBuilder)
+	{
+		static const FName MrgTokenAutocompleteMetaDataKey(TEXT("MrgTokenAutocomplete"));
+
+		// Find all of the properties that should have the autocomplete widget used on them. Note that this won't find properties in structs.
+		// Structs need separate details customizations that add autocomplete if necessary.
+		TArray<TSharedPtr<IPropertyHandle>> AutocompleteProperties;
+		TArray<FName> CategoryNames;
+		DetailBuilder.GetCategoryNames(CategoryNames);
+		for (const FName& CategoryName : CategoryNames)
+		{
+			IDetailCategoryBuilder& CategoryBuilder = DetailBuilder.EditCategory(CategoryName);
+
+			TArray<TSharedRef<IPropertyHandle>> CategoryProperties;
+			CategoryBuilder.GetDefaultProperties(CategoryProperties);
+
+			for (const TSharedRef<IPropertyHandle>& CategoryProperty : CategoryProperties)
+			{
+				if (CategoryProperty->HasMetaData(MrgTokenAutocompleteMetaDataKey))
+				{
+					AutocompleteProperties.Add(CategoryProperty);
+				}
+			}
+		}
+
+		// Add the autocomplete widget to all properties that have the appropriate UPROPERTY metadata
+		for (const TSharedPtr<IPropertyHandle>& PropertyToModify : AutocompleteProperties)
+		{
+			if (IDetailPropertyRow* PropertyRow = DetailBuilder.EditDefaultProperty(PropertyToModify))
+			{
+				UE::MovieRenderPipelineEditor::Private::AddTokenAutocompleteToPropertyRow(PropertyRow);
+			}
+		}
+	}
 };
 
 #undef LOCTEXT_NAMESPACE
\ No newline at end of file
diff --git a/Engine/Plugins/MovieScene/MovieRenderPipeline/Source/MovieRenderPipelineEditor/Private/Customizations/Graph/MovieGraphPostProcessPassCustomization.h b/Engine/Plugins/MovieScene/MovieRenderPipeline/Source/MovieRenderPipelineEditor/Private/Customizations/Graph/MovieGraphPostProcessPassCustomization.h
new file mode 100644
--- /dev/null
+++ b/Engine/Plugins/MovieScene/MovieRenderPipeline/Source/MovieRenderPipelineEditor/Private/Customizations/Graph/MovieGraphPostProcessPassCustomization.h
@@ -0,0 +1,69 @@
+﻿// Copyright Epic Games, Inc. All Rights Reserved.
+
+#pragma once
+
+#include "DetailWidgetRow.h"
+#include "IDetailChildrenBuilder.h"
+#include "IPropertyTypeCustomization.h"
+#include "MovieGraphCustomizationUtils.h"
+#include "MoviePipelineDeferredPasses.h"
+#include "PropertyHandle.h"
+
+#define LOCTEXT_NAMESPACE "MoviePipelineEditor"
+
+/** Customize how post process passes appear in the details panel. */
+class FMovieGraphPostProcessPassCustomization final : public IPropertyTypeCustomization
+{
+public:
+	static TSharedRef<IPropertyTypeCustomization> MakeInstance()
+	{
+		return MakeShared<FMovieGraphPostProcessPassCustomization>();
+	}
+
+protected:
+	//~ Begin IPropertyTypeCustomization interface
+	virtual void CustomizeHeader(
+		TSharedRef<IPropertyHandle> InStructPropertyHandle,
+		FDetailWidgetRow& HeaderRow,
+		IPropertyTypeCustomizationUtils& CustomizationUtils) override
+	{
+		HeaderRow
+			.NameContent()
+			[
+				InStructPropertyHandle->CreatePropertyNameWidget()
+			]
+			.ValueContent()
+			[
+				InStructPropertyHandle->CreatePropertyValueWidget()
+			];
+	}
+
+	virtual void CustomizeChildren(
+		TSharedRef<IPropertyHandle> InStructPropertyHandle,
+		IDetailChildrenBuilder& StructBuilder,
+		IPropertyTypeCustomizationUtils& StructCustomizationUtils) override
+	{
+		if (!InStructPropertyHandle->IsValidHandle())
+		{
+			return;
+		}
+
+		uint32 NumChildren;
+		InStructPropertyHandle->GetNumChildren(NumChildren);
+
+		for (uint32 Index = 0; Index < NumChildren; ++Index)
+		{
+			TSharedPtr<IPropertyHandle> ChildHandle = InStructPropertyHandle->GetChildHandle(Index);
+			IDetailPropertyRow& PropertyRow = StructBuilder.AddProperty(ChildHandle.ToSharedRef());
+
+			// Give the "Name" property the token autocomplete widget
+			if (ChildHandle->GetProperty()->GetFName() == GET_MEMBER_NAME_CHECKED(FMoviePipelinePostProcessPass, Name))
+			{
+				UE::MovieRenderPipelineEditor::Private::AddTokenAutocompleteToPropertyRow(&PropertyRow);
+			}
+		}
+	}
+	//~ End IPropertyTypeCustomization interface
+};
+
+#undef LOCTEXT_NAMESPACE
\ No newline at end of file
diff --git a/Engine/Plugins/MovieScene/MovieRenderPipeline/Source/MovieRenderPipelineEditor/Private/Graph/MovieGraphAssetToolkit.cpp b/Engine/Plugins/MovieScene/MovieRenderPipeline/Source/MovieRenderPipelineEditor/Private/Graph/MovieGraphAssetToolkit.cpp
--- a/Engine/Plugins/MovieScene/MovieRenderPipeline/Source/MovieRenderPipelineEditor/Private/Graph/MovieGraphAssetToolkit.cpp
+++ b/Engine/Plugins/MovieScene/MovieRenderPipeline/Source/MovieRenderPipelineEditor/Private/Graph/MovieGraphAssetToolkit.cpp
@@ -5,7 +5,6 @@
 #include "Customizations/Graph/MovieGraphApplyCVarPresetNodeCustomization.h"
 #include "Customizations/Graph/MovieGraphCollectionsCustomization.h"
 #include "Customizations/Graph/MovieGraphMultiLayerEXRNodeCustomization.h"
-#include "Customizations/Graph/MovieGraphFormatTokenCustomization.h"
 #include "Customizations/Graph/MovieGraphMaterialParameterCollectionModifierNodeCustomization.h"
 #include "Customizations/Graph/MovieGraphMemberCustomization.h"
 #include "Customizations/Graph/MovieGraphMetadataAttributeCustomization.h"
@@ -13,6 +12,7 @@
 #include "Customizations/Graph/MovieGraphNamedResolutionCustomization.h"
 #include "Customizations/Graph/MovieGraphNodeCustomization.h"
 #include "Customizations/Graph/MovieGraphPathTracedRendererNodeCustomization.h"
+#include "Customizations/Graph/MovieGraphPostProcessPassCustomization.h"
 #include "Customizations/Graph/MovieGraphSelectNodeCustomization.h"
 #include "Customizations/Graph/MovieGraphSetCVarValueNodeCustomization.h"
 #include "Customizations/Graph/MovieGraphShowFlagsCustomization.h"
@@ -436,18 +436,10 @@ TSharedRef<SDockTab> FMovieGraphAssetToolkit::SpawnTab_RenderGraphDetails(const
 	SelectedGraphObjectsDetailsWidget->RegisterInstancedCustomPropertyLayout(
 		UMovieGraphModifierNode::StaticClass(),
 		FOnGetDetailCustomizationInstance::CreateStatic(&FMovieGraphModifiersCustomization::MakeInstance));
-	
-	SelectedGraphObjectsDetailsWidget->RegisterInstancedCustomPropertyLayout(
-		UMovieGraphFileOutputNode::StaticClass(),
-		FOnGetDetailCustomizationInstance::CreateStatic(&FMovieGraphFormatTokenCustomization::MakeInstance));
-
-	SelectedGraphObjectsDetailsWidget->RegisterInstancedCustomPropertyLayout(
-		UMovieGraphCommandLineEncoderNode::StaticClass(),
-		FOnGetDetailCustomizationInstance::CreateStatic(&FMovieGraphFormatTokenCustomization::MakeInstance));
 
-	SelectedGraphObjectsDetailsWidget->RegisterInstancedCustomPropertyLayout(
-		UMovieGraphDebugSettingNode::StaticClass(),
-		FOnGetDetailCustomizationInstance::CreateStatic(&FMovieGraphFormatTokenCustomization::MakeInstance));
+	SelectedGraphObjectsDetailsWidget->RegisterInstancedCustomPropertyTypeLayout(
+		FMoviePipelinePostProcessPass::StaticStruct()->GetFName(),
+		FOnGetPropertyTypeCustomizationInstance::CreateStatic(&FMovieGraphPostProcessPassCustomization::MakeInstance));
 	
 	SelectedGraphObjectsDetailsWidget->RegisterInstancedCustomPropertyLayout(
 		UMovieGraphPathTracerRenderPassNode::StaticClass(),
diff --git a/Engine/Plugins/MovieScene/MovieRenderPipeline/Source/MovieRenderPipelineEditor/Private/MovieRenderPipelineStyle.cpp b/Engine/Plugins/MovieScene/MovieRenderPipeline/Source/MovieRenderPipelineEditor/Private/MovieRenderPipelineStyle.cpp
--- a/Engine/Plugins/MovieScene/MovieRenderPipeline/Source/MovieRenderPipelineEditor/Private/MovieRenderPipelineStyle.cpp
+++ b/Engine/Plugins/MovieScene/MovieRenderPipeline/Source/MovieRenderPipelineEditor/Private/MovieRenderPipelineStyle.cpp
@@ -101,6 +101,17 @@ FMovieRenderPipelineStyle::FMovieRenderPipelineStyle()
 		.SetColorAndOpacity(FSlateColor::UseSubduedForeground())
 	);
 
+	const FEditableTextBoxStyle NormalEditableTextBoxStyle = FCoreStyle::Get().GetWidgetStyle<FEditableTextBoxStyle>("NormalEditableTextBox");
+	const FTextBlockStyle NormalTextBlockStyle = FAppStyle::GetWidgetStyle<FTextBlockStyle>("NormalText");
+	
+	FEditableTextBoxStyle AutocompleteEditableText = FEditableTextBoxStyle(NormalEditableTextBoxStyle)
+		.SetFont(FAppStyle::GetFontStyle(TEXT("PropertyWindow.NormalFont")));
+	Set("MovieRenderPipeline.AutocompleteTextStyle", AutocompleteEditableText);
+	
+	FTextBlockStyle AutocompleteArgumentEditableText = FTextBlockStyle(NormalTextBlockStyle)
+		.SetFont(FAppStyle::GetFontStyle(TEXT("PropertyWindow.BoldFont")));
+	Set("MovieRenderPipeline.AutocompleteArgumentTextStyle", AutocompleteArgumentEditableText);
+
 	Set("MovieRenderPipeline.TabIcon", new FSlateImageBrush(RootToContentDir(TEXT("TabIcon_24x.png")), Icon16x16));
 
 	Set("MovieRenderGraph.CollectionsTree.SmallComboButton", FComboButtonStyle(FAppStyle::Get().GetWidgetStyle<FComboButtonStyle>("SimpleComboButton"))
diff --git a/Engine/Plugins/MovieScene/MovieRenderPipeline/Source/MovieRenderPipelineMP4Encoder/Public/Graph/MovieGraphMP4EncoderNode.h b/Engine/Plugins/MovieScene/MovieRenderPipeline/Source/MovieRenderPipelineMP4Encoder/Public/Graph/MovieGraphMP4EncoderNode.h
--- a/Engine/Plugins/MovieScene/MovieRenderPipeline/Source/MovieRenderPipelineMP4Encoder/Public/Graph/MovieGraphMP4EncoderNode.h
+++ b/Engine/Plugins/MovieScene/MovieRenderPipeline/Source/MovieRenderPipelineMP4Encoder/Public/Graph/MovieGraphMP4EncoderNode.h
@@ -191,7 +191,7 @@ public:
 	bool bCompositeOntoFinalImage = true;
 
 	/** The file name format used for writing out the burn-in video (if not composited). */
-	UPROPERTY(EditAnywhere, BlueprintReadWrite, Category = "Burn In", meta = (EditCondition = "bOverride_BurnInFileNameFormat"))
+	UPROPERTY(EditAnywhere, BlueprintReadWrite, Category = "Burn In", meta = (EditCondition = "bOverride_BurnInFileNameFormat", MrgTokenAutocomplete))
 	FString BurnInFileNameFormat = TEXT("{sequence_name}.{layer_name}.{renderer_name}");
 };
 
diff --git a/Engine/Plugins/MovieScene/MovieRenderPipeline/Source/MovieRenderPipelineRenderPasses/Public/Graph/Nodes/MovieGraphDeferredPassNode.h b/Engine/Plugins/MovieScene/MovieRenderPipeline/Source/MovieRenderPipelineRenderPasses/Public/Graph/Nodes/MovieGraphDeferredPassNode.h
--- a/Engine/Plugins/MovieScene/MovieRenderPipeline/Source/MovieRenderPipelineRenderPasses/Public/Graph/Nodes/MovieGraphDeferredPassNode.h
+++ b/Engine/Plugins/MovieScene/MovieRenderPipeline/Source/MovieRenderPipelineRenderPasses/Public/Graph/Nodes/MovieGraphDeferredPassNode.h
@@ -154,7 +154,7 @@ public:
 	/**
 	 * If specified, overrides the output node's file name format, and uses this file name format for Post Process Material passes instead.
 	 */
-	UPROPERTY(EditAnywhere, BlueprintReadWrite, Category = "Post Process Materials", DisplayName = "PPM File Name Format", meta=(EditCondition="bOverride_PPMFileNameFormat"))
+	UPROPERTY(EditAnywhere, BlueprintReadWrite, Category = "Post Process Materials", DisplayName = "PPM File Name Format", meta=(EditCondition="bOverride_PPMFileNameFormat", MrgTokenAutocomplete))
 	FString PPMFileNameFormat;
 
 	/**
diff --git a/Engine/Plugins/MovieScene/MovieRenderPipeline/Source/MovieRenderPipelineRenderPasses/Public/Graph/Nodes/MovieGraphPathTracerPassNode.h b/Engine/Plugins/MovieScene/MovieRenderPipeline/Source/MovieRenderPipelineRenderPasses/Public/Graph/Nodes/MovieGraphPathTracerPassNode.h
--- a/Engine/Plugins/MovieScene/MovieRenderPipeline/Source/MovieRenderPipelineRenderPasses/Public/Graph/Nodes/MovieGraphPathTracerPassNode.h
+++ b/Engine/Plugins/MovieScene/MovieRenderPipeline/Source/MovieRenderPipelineRenderPasses/Public/Graph/Nodes/MovieGraphPathTracerPassNode.h
@@ -263,7 +263,7 @@ public:
 	/**
 	 * If specified, overrides the output node's file name format, and uses this file name format for Post Process Material passes instead.
 	 */
-	UPROPERTY(EditAnywhere, BlueprintReadWrite, Category = "Post Process Materials", DisplayName = "PPM File Name Format", meta=(EditCondition="bOverride_PPMFileNameFormat"))
+	UPROPERTY(EditAnywhere, BlueprintReadWrite, Category = "Post Process Materials", DisplayName = "PPM File Name Format", meta=(EditCondition="bOverride_PPMFileNameFormat", MrgTokenAutocomplete))
 	FString PPMFileNameFormat;
 
 	/**
diff --git a/Engine/Plugins/MovieScene/MovieRenderPipeline/Source/MovieRenderPipelineRenderPasses/Public/MovieGraphImageSequenceOutputNode.h b/Engine/Plugins/MovieScene/MovieRenderPipeline/Source/MovieRenderPipelineRenderPasses/Public/MovieGraphImageSequenceOutputNode.h
--- a/Engine/Plugins/MovieScene/MovieRenderPipeline/Source/MovieRenderPipelineRenderPasses/Public/MovieGraphImageSequenceOutputNode.h
+++ b/Engine/Plugins/MovieScene/MovieRenderPipeline/Source/MovieRenderPipelineRenderPasses/Public/MovieGraphImageSequenceOutputNode.h
@@ -196,7 +196,7 @@ public:
 	bool bCompositeOntoFinalImage = true;
 
 	/** The file name format used for writing out the burn-in image (if not composited). */
-	UPROPERTY(EditAnywhere, BlueprintReadWrite, Category = "Burn In", meta = (EditCondition = "bOverride_BurnInFileNameFormat"))
+	UPROPERTY(EditAnywhere, BlueprintReadWrite, Category = "Burn In", meta = (EditCondition = "bOverride_BurnInFileNameFormat", MrgTokenAutocomplete))
 	FString BurnInFileNameFormat = TEXT("{sequence_name}.{layer_name}.BurnIn.{frame_number}");
 
 protected:
@@ -433,7 +433,7 @@ public:
 	bool bCompositeOntoFinalImage = true;
 
 	/** The file name format used for writing out the burn-in image (if not composited). */
-	UPROPERTY(EditAnywhere, BlueprintReadWrite, Category = "Burn In", meta = (EditCondition = "bOverride_BurnInFileNameFormat"))
+	UPROPERTY(EditAnywhere, BlueprintReadWrite, Category = "Burn In", meta = (EditCondition = "bOverride_BurnInFileNameFormat", MrgTokenAutocomplete))
 	FString BurnInFileNameFormat = TEXT("{sequence_name}.{layer_name}.BurnIn.{frame_number}");
 };
 
@@ -519,7 +519,7 @@ public:
 	bool bCompositeOntoFinalImage = true;
 
 	/** The file name format used for writing out the burn-in image (if not composited). */
-	UPROPERTY(EditAnywhere, BlueprintReadWrite, Category = "Burn In", meta = (EditCondition = "bOverride_BurnInFileNameFormat"))
+	UPROPERTY(EditAnywhere, BlueprintReadWrite, Category = "Burn In", meta = (EditCondition = "bOverride_BurnInFileNameFormat", MrgTokenAutocomplete))
 	FString BurnInFileNameFormat = TEXT("{sequence_name}.{layer_name}.BurnIn.{frame_number}");
 };
 
@@ -606,7 +606,7 @@ public:
 	bool bCompositeOntoFinalImage = true;
 
 	/** The file name format used for writing out the burn-in image (if not composited). */
-	UPROPERTY(EditAnywhere, BlueprintReadWrite, Category = "Burn In", meta = (EditCondition = "bOverride_BurnInFileNameFormat"))
+	UPROPERTY(EditAnywhere, BlueprintReadWrite, Category = "Burn In", meta = (EditCondition = "bOverride_BurnInFileNameFormat", MrgTokenAutocomplete))
 	FString BurnInFileNameFormat = TEXT("{sequence_name}.{layer_name}.BurnIn.{frame_number}");
 };
 
diff --git a/Engine/Plugins/MovieScene/MovieRenderPipeline/Source/MovieRenderPipelineRenderPasses/Public/MoviePipelineDeferredPasses.h b/Engine/Plugins/MovieScene/MovieRenderPipeline/Source/MovieRenderPipelineRenderPasses/Public/MoviePipelineDeferredPasses.h
--- a/Engine/Plugins/MovieScene/MovieRenderPipeline/Source/MovieRenderPipelineRenderPasses/Public/MoviePipelineDeferredPasses.h
+++ b/Engine/Plugins/MovieScene/MovieRenderPipeline/Source/MovieRenderPipelineRenderPasses/Public/MoviePipelineDeferredPasses.h
@@ -30,7 +30,7 @@ public:
 	 * token. For Movie Render Graph, the material name will be used in {renderer_sub_name}. If a name is not specified here, the full name of the material
 	 * will be used in these tokens instead.
 	 */
-	UPROPERTY(EditAnywhere, BlueprintReadWrite, Category = "Settings")
+	UPROPERTY(EditAnywhere, BlueprintReadWrite, Category = "Settings", meta = (MrgTokenAutocomplete))
 	FString Name;
 
 	/** 
