MoviePipeline: Ensure that DirectoryOpenBehavior is initialized properly in the Play After Render settings.

#rnx

diff --git a/Engine/Plugins/MovieScene/MovieRenderPipeline/Source/MovieRenderPipelineCore/Private/MoviePipelinePostRenderSettings.cpp b/Engine/Plugins/MovieScene/MovieRenderPipeline/Source/MovieRenderPipelineCore/Private/MoviePipelinePostRenderSettings.cpp
--- a/Engine/Plugins/MovieScene/MovieRenderPipeline/Source/MovieRenderPipelineCore/Private/MoviePipelinePostRenderSettings.cpp
+++ b/Engine/Plugins/MovieScene/MovieRenderPipeline/Source/MovieRenderPipelineCore/Private/MoviePipelinePostRenderSettings.cpp
@@ -22,6 +22,7 @@ FMovieGraphPostRenderImageSequencePlayOptions::FMovieGraphPostRenderImageSequenc
 FMovieGraphPostRenderSettings::FMovieGraphPostRenderSettings()
 	: OutputTypePriorityOrder({TEXT("EXR"), TEXT("PNG"), TEXT("JPEG"), TEXT("MP4"), TEXT("BMP"), TEXT("MOV"), TEXT("MXF")})
 	, OutputTypePlayback(EMovieGraphOutputTypePlayback::UsePriorityOrder)
+	, DirectoryOpenBehavior(EMovieGraphDirectoryOpenBehavior::OpenDirectoryForAllJobs)
 {
 	
 }
