MoviePipeline: Added tooltip diagrams on ShutterTiming setting.

diff --git a/Engine/Plugins/MovieScene/MovieRenderPipeline/Source/MovieRenderPipelineCore/Public/Graph/Nodes/MovieGraphCameraNode.h b/Engine/Plugins/MovieScene/MovieRenderPipeline/Source/MovieRenderPipelineCore/Public/Graph/Nodes/MovieGraphCameraNode.h
--- a/Engine/Plugins/MovieScene/MovieRenderPipeline/Source/MovieRenderPipelineCore/Public/Graph/Nodes/MovieGraphCameraNode.h
+++ b/Engine/Plugins/MovieScene/MovieRenderPipeline/Source/MovieRenderPipelineCore/Public/Graph/Nodes/MovieGraphCameraNode.h
@@ -48,6 +48,23 @@ public:
 	* before and right up to frame N. When set to FrameCenter, the motion represents half the time before the
 	* frame and half the time after the frame. When set to FrameOpen, the motion represents the time from 
 	* Frame N onwards.
+	* 
+	* As an example, here are 4 temporal subsamples with 180 degree shutter angle (50% motion blur) between frames A & B:
+	*
+	* FrameClose
+	* F-------F-------F
+	* ----1234---------
+	* ----AAAA---------
+	*
+	* FrameCenter
+	* F-------F-------F
+	* ------1234-------
+	* ------AABB-------
+	*
+	* FrameOpen
+	* F-------F-------F
+	* --------1234-----
+	* --------BBBB-----
 	*/
 	UPROPERTY(EditAnywhere, BlueprintReadWrite, Category = "Settings", meta = (EditCondition = "bOverride_ShutterTiming"))
 	EMoviePipelineShutterTiming ShutterTiming;
diff --git a/Engine/Plugins/MovieScene/MovieRenderPipeline/Source/MovieRenderPipelineCore/Public/MoviePipelineCameraSetting.h b/Engine/Plugins/MovieScene/MovieRenderPipeline/Source/MovieRenderPipelineCore/Public/MoviePipelineCameraSetting.h
--- a/Engine/Plugins/MovieScene/MovieRenderPipeline/Source/MovieRenderPipelineCore/Public/MoviePipelineCameraSetting.h
+++ b/Engine/Plugins/MovieScene/MovieRenderPipeline/Source/MovieRenderPipelineCore/Public/MoviePipelineCameraSetting.h
@@ -35,6 +35,23 @@ public:
 	* before and right up to frame N. When set to FrameCenter, the motion represents half the time before the
 	* frame and half the time after the frame. When set to FrameOpen, the motion represents the time from 
 	* Frame N onwards.
+	*
+	* As an example, here are 4 temporal subsamples with 180 degree shutter angle (50% motion blur) between frames A & B:
+	*
+	* FrameClose
+	* F-------F-------F
+	* ----1234---------
+	* ----AAAA---------
+	*
+	* FrameCenter
+	* F-------F-------F
+	* ------1234-------
+	* ------AABB-------
+	*
+	* FrameOpen
+	* F-------F-------F
+	* --------1234-----
+	* --------BBBB-----
 	*/
 	UPROPERTY(EditAnywhere, BlueprintReadWrite, Category = "Camera Settings")
 	EMoviePipelineShutterTiming ShutterTiming;
