MoviePipeline: Updated the Execute Script node to include two different modes, "Editor Only" and "Editor + Runtime". This was primarily done so the Asset Referencing Restrictions plugin doesn't flag Python scripts (this requires the new Editor Only mode to be used -- Python scripts cannot be referenced in Runtime builds).

Also enabled MovieGraphScriptBase for use in Blueprints, meaning Blueprints can now be used for scripting within MRG via the Execute Script node.

diff --git a/Engine/Plugins/MovieScene/MovieRenderPipeline/Source/MovieRenderPipelineCore/Private/Graph/Nodes/MovieGraphExecuteScriptNode.cpp b/Engine/Plugins/MovieScene/MovieRenderPipeline/Source/MovieRenderPipelineCore/Private/Graph/Nodes/MovieGraphExecuteScriptNode.cpp
--- a/Engine/Plugins/MovieScene/MovieRenderPipeline/Source/MovieRenderPipelineCore/Private/Graph/Nodes/MovieGraphExecuteScriptNode.cpp
+++ b/Engine/Plugins/MovieScene/MovieRenderPipeline/Source/MovieRenderPipelineCore/Private/Graph/Nodes/MovieGraphExecuteScriptNode.cpp
@@ -3,6 +3,7 @@
 #include "Graph/Nodes/MovieGraphExecuteScriptNode.h"
 
 #include "MoviePipelineTelemetry.h"
+#include "MovieRenderPipelineCoreModule.h"
 #include "Styling/AppStyle.h"
 #include "UObject/Package.h"
 
@@ -36,7 +37,26 @@ FSlateIcon UMovieGraphExecuteScriptNode::GetIconAndTint(FLinearColor& OutColor)
 
 UMovieGraphScriptBase* UMovieGraphExecuteScriptNode::AllocateScriptInstance() const
 {
-	UClass* NewInstanceType = Script.TryLoadClass<UMovieGraphScriptBase>();
+	const UClass* NewInstanceType = nullptr;
+
+	if (Mode == EMovieGraphExecuteScriptMode::EditorOnly)
+	{
+#if WITH_EDITORONLY_DATA
+		NewInstanceType = EditorOnlyScript.TryLoadClass<UMovieGraphScriptBase>();
+#else
+		UE_LOG(LogMovieRenderPipeline, Warning, TEXT("Trying to use an editor-only script in a non-editor build. Script will not execute."));
+#endif
+	}
+	else if (Mode == EMovieGraphExecuteScriptMode::EditorAndRuntime)
+	{
+		NewInstanceType = EditorAndRuntimeScript.TryLoadClass<UMovieGraphScriptBase>();
+	}
+	else
+	{
+		// Invalid mode
+		check(false);
+	}
+
 	if (NewInstanceType)
 	{
 		UMovieGraphScriptBase* NewInstance = NewObject<UMovieGraphScriptBase>(GetTransientPackage(), NewInstanceType);
@@ -46,6 +66,50 @@ UMovieGraphScriptBase* UMovieGraphExecuteScriptNode::AllocateScriptInstance() co
 	return nullptr;
 }
 
+void UMovieGraphExecuteScriptNode::PostLoad()
+{
+	Super::PostLoad();
+
+	// If loading an older Execute Script node which uses the Script property, migrate it to the EditorAndRuntimeScript property.
+	PRAGMA_DISABLE_DEPRECATION_WARNINGS
+	if (Script.IsValid())
+	{
+		Mode = EMovieGraphExecuteScriptMode::EditorAndRuntime;
+		EditorAndRuntimeScript = Script;
+		bOverride_EditorAndRuntimeScript = bOverride_Script;
+		Script = FSoftClassPath();
+	}
+	PRAGMA_ENABLE_DEPRECATION_WARNINGS
+}
+
+#if WITH_EDITOR
+void UMovieGraphExecuteScriptNode::PostEditChangeProperty(FPropertyChangedEvent& PropertyChangedEvent)
+{
+	Super::PostEditChangeProperty(PropertyChangedEvent);
+
+	// If switching modes, clear out the script property for the mode being switched from. This is particularly important when switching
+	// from Editor + Runtime mode to Editor-Only mode. If the user is intentionally flagging the script as Editor-Only, but previously
+	// had an Editor + Runtime script set, the EditorAndRuntimeScript *will still have a reference to the script*, which can trip up
+	// the Asset Referencing Restrictions plugin. Clearing out the script property when switching modes fixes this issue.
+	if (PropertyChangedEvent.GetPropertyName() == GET_MEMBER_NAME_CHECKED(UMovieGraphExecuteScriptNode, Mode))
+	{
+		if (Mode == EMovieGraphExecuteScriptMode::EditorOnly)
+		{
+			EditorAndRuntimeScript = FSoftClassPath();
+		}
+		else if (Mode == EMovieGraphExecuteScriptMode::EditorAndRuntime)
+		{
+			EditorOnlyScript = FSoftClassPath();
+		}
+		else
+		{
+			// Invalid mode
+			check(false);
+		}
+	}
+}
+#endif	// WITH_EDITOR
+
 void UMovieGraphExecuteScriptNode::UpdateTelemetry(FMoviePipelineShotRenderTelemetry* InTelemetry) const
 {
 	InTelemetry->bUsesScripting = true;
diff --git a/Engine/Plugins/MovieScene/MovieRenderPipeline/Source/MovieRenderPipelineCore/Public/Graph/Nodes/MovieGraphExecuteScriptNode.h b/Engine/Plugins/MovieScene/MovieRenderPipeline/Source/MovieRenderPipelineCore/Public/Graph/Nodes/MovieGraphExecuteScriptNode.h
--- a/Engine/Plugins/MovieScene/MovieRenderPipeline/Source/MovieRenderPipelineCore/Public/Graph/Nodes/MovieGraphExecuteScriptNode.h
+++ b/Engine/Plugins/MovieScene/MovieRenderPipeline/Source/MovieRenderPipelineCore/Public/Graph/Nodes/MovieGraphExecuteScriptNode.h
@@ -14,12 +14,13 @@
 * (named 'Execute Script' in the UI) in your Movie Graph Configuration asset,
 * and choosing the type of your script, an instance will be created when
 * the job starts. This instance will persist until the end of the job. The instance
-* will recieve callbacks during the render which contain data that you can 
+* will receive callbacks during the render which contain data that you can 
 * manipulate with your script, or data that was generated by the pipeline.
 *
-* This can be implemented in either C++ or Python. For C++ you can use standard
-* inheritance but for Python you need to write a UClass implemented in Python.
-* This has slightly different syntax than standard Python;
+* This can be implemented in C++, Blueprints, or Python. For C++ you can use standard
+* inheritance. Blueprints should use the MovieGraphScriptBase class. For Python,
+* you need to write a UClass implemented in Python. This has slightly different
+* syntax than standard Python (see below).
 * 
 * In this example, create a "Python" folder in your Content folder. Then create
 * a file named "MovieGraphPipelineScriptExample.py", and a file named "init_unreal.py"
@@ -45,7 +46,7 @@
 * See /Engine/Plugins/MovieScene/MovieRenderPipeline/Content/Python/MovieGraphScriptNodeExample.py 
 * for more python examples.
 */
-UCLASS(MinimalAPI, BlueprintType, Abstract)
+UCLASS(MinimalAPI, BlueprintType, Blueprintable, Abstract)
 class UMovieGraphScriptBase : public UObject
 {
 	GENERATED_BODY()
@@ -170,6 +171,20 @@ public:
 	}
 };
 
+/** Designates the types of builds a script (specified in an Execute Script node) is allowed to execute in. */
+UENUM(BlueprintType)
+enum class EMovieGraphExecuteScriptMode : uint8
+{
+	/** The script should only run in an editor context. */
+	EditorOnly = 0,
+
+	/** The script can run both in the editor and in a runtime build. */
+	EditorAndRuntime = 1,
+
+	/** Invalid value, keep as the last entry in the enum. */
+	Count UMETA(Hidden)
+};
+
 /** 
  This node allows users to run code before and after a movie graph is being executed. Example use cases
  would be to run custom logic that overrides the output directory, or to trigger an event after a shot finishes.
@@ -194,16 +209,33 @@ public:
 	UE_API virtual FSlateIcon GetIconAndTint(FLinearColor& OutColor) const override;
 #endif
 
-	/** 
-	* Allocate an instance of the specifed script type.
-	* Can return nullptr.
-	*/
+	/** Allocates an instance of the script type specified on the node. Will return nullptr if the script type is invalid. */
 	UE_API UMovieGraphScriptBase* AllocateScriptInstance() const;
 
+protected:
+	//~ Begin UObject Interface
+	virtual void PostLoad() override;
+
+#if WITH_EDITOR
+	virtual void PostEditChangeProperty(FPropertyChangedEvent& PropertyChangedEvent) override;
+#endif
+	//~ End UObject Interface
+
 public:
-	UPROPERTY(EditAnywhere, BlueprintReadWrite, Category = Overrides, meta = (InlineEditConditionToggle))
+	UE_DEPRECATED(5.8, "Use the bOverride_EditorOnlyScript property or bOverride_EditorAndRuntimeScript property instead.")
+	UPROPERTY(BlueprintReadWrite, Category = Overrides, meta = (InlineEditConditionToggle, DeprecatedProperty, DeprecationMessage = "Use the bOverride_EditorOnlyScript property or bOverride_EditorAndRuntimeScript property instead."))
 	uint8 bOverride_Script: 1;
-	
+
+#if WITH_EDITORONLY_DATA
+	UPROPERTY(EditAnywhere, BlueprintReadWrite, Category = Overrides, meta = (InlineEditConditionToggle))
+	uint8 bOverride_EditorOnlyScript : 1;
+#endif
+
+	UPROPERTY(EditAnywhere, BlueprintReadWrite, Category = Overrides, meta = (InlineEditConditionToggle))
+	uint8 bOverride_EditorAndRuntimeScript : 1;
+
+	UPROPERTY(EditAnywhere, BlueprintReadWrite, Category = Overrides, meta = (InlineEditConditionToggle))
+	uint8 bOverride_Mode : 1;
 
 	/**
 	* Which script should this node instantiate? The class will be instantiated when the job starts,
@@ -211,8 +243,40 @@ public:
 	* movie render. Scripts can be implemented in Python or C++, see MovieGraphExecuteScriptNode.h
 	* for more details.
 	*/
-	UPROPERTY(EditAnywhere, BlueprintReadWrite, Category = "Settings", meta = (EditCondition = "bOverride_Script", MetaClass = "/Script/MovieRenderPipelineCore.MovieGraphScriptBase"))
+	UE_DEPRECATED(5.8, "Use the EditorOnlyScript property or EditorAndRuntimeScript property instead.")
+	UPROPERTY(BlueprintReadWrite, Category = "Settings", meta = (EditCondition = "bOverride_Script", MetaClass = "/Script/MovieRenderPipelineCore.MovieGraphScriptBase", DeprecatedProperty, DeprecationMessage = "Use the EditorOnlyScript property or EditorAndRuntimeScript property instead."))
 	FSoftClassPath Script;
+
+	/** The context that this script is allowed to execute in. */
+	UPROPERTY(EditAnywhere, BlueprintReadWrite, Category = "Settings", meta = (EditCondition = "bOverride_Mode"))
+	EMovieGraphExecuteScriptMode Mode = EMovieGraphExecuteScriptMode::EditorOnly;
+
+#if WITH_EDITORONLY_DATA
+	/**
+	 * The script that this node should run. The script will be created when the job starts, and different functions will be called on the
+	 * script throughout the lifetime of the render. Scripts can be implemented in Python, Blueprints, or C++; see MovieGraphExecuteScriptNode.h
+	 * for more details.
+	 *
+	 * Editor-only scripts will run in editor-based renders and not runtime builds (if attempting to use it in a runtime build, the script will not
+	 * execute, and a warning will be logged). This is mainly intended for typical MRG renders using a Python-based script (but the script can be
+	 * C++ or Blueprint-based as well). Using this script type allows the Asset Referencing Restrictions plugin to properly detect that this script
+	 * will not execute in a runtime build.
+	 */
+	UPROPERTY(EditAnywhere, BlueprintReadWrite, Category = "Settings", DisplayName = "Editor-Only Script", meta = (EditCondition = "bOverride_EditorOnlyScript", MetaClass = "/Script/MovieRenderPipelineCore.MovieGraphScriptBase"))
+	FSoftClassPath EditorOnlyScript;
+#endif
+
+	/**
+	 * The script that this node should run. The script will be created when the job starts, and different functions will be called on the
+	 * script throughout the lifetime of the render. Scripts can be implemented in Python, Blueprints, or C++; see MovieGraphExecuteScriptNode.h
+	 * for more details.
+	 *
+	 * Editor and Runtime scripts will run both in the editor and in runtime builds. Note that if a Python-based script type is being used in a
+	 * runtime build, the Asset Referencing Restrictions plugin will flag an issue. This script type is mostly intended for use with Blueprints and
+	 * C++-based scripts.
+	 */
+	UPROPERTY(EditAnywhere, BlueprintReadWrite, Category = "Settings", DisplayName = "Editor + Runtime Script", meta = (EditCondition = "bOverride_EditorAndRuntimeScript", MetaClass = "/Script/MovieRenderPipelineCore.MovieGraphScriptBase"))
+	FSoftClassPath EditorAndRuntimeScript;
 };
 
 #undef UE_API
diff --git a/Engine/Plugins/MovieScene/MovieRenderPipeline/Source/MovieRenderPipelineEditor/Private/Customizations/Graph/MovieGraphExecuteScriptNodeCustomization.h b/Engine/Plugins/MovieScene/MovieRenderPipeline/Source/MovieRenderPipelineEditor/Private/Customizations/Graph/MovieGraphExecuteScriptNodeCustomization.h
new file mode 100644
--- /dev/null
+++ b/Engine/Plugins/MovieScene/MovieRenderPipeline/Source/MovieRenderPipelineEditor/Private/Customizations/Graph/MovieGraphExecuteScriptNodeCustomization.h
@@ -0,0 +1,75 @@
+﻿// Copyright Epic Games, Inc. All Rights Reserved.
+
+#pragma once
+
+#include "DetailLayoutBuilder.h"
+#include "Graph/Nodes/MovieGraphExecuteScriptNode.h"
+#include "IDetailCustomization.h"
+#include "PropertyHandle.h"
+
+#define LOCTEXT_NAMESPACE "MoviePipelineEditor"
+
+/** Customize how the Execute Script node appears in the details panel. */
+class FMovieGraphExecuteScriptNodeCustomization : public IDetailCustomization
+{
+public:
+	static TSharedRef<IDetailCustomization> MakeInstance()
+	{
+		return MakeShared<FMovieGraphExecuteScriptNodeCustomization>();
+	}
+
+protected:
+	//~ Begin IDetailCustomization interface
+	virtual void CustomizeDetails(const TSharedPtr<IDetailLayoutBuilder>& DetailBuilder) override
+	{
+		CustomizeDetails(*DetailBuilder);
+	}
+	
+	virtual void CustomizeDetails(IDetailLayoutBuilder& DetailBuilder) override
+	{
+		const TSharedRef<IPropertyHandle> ModeHandle = DetailBuilder.GetProperty(GET_MEMBER_NAME_CHECKED(UMovieGraphExecuteScriptNode, Mode));
+		const TSharedRef<IPropertyHandle> EditorOnlyHandle = DetailBuilder.GetProperty(GET_MEMBER_NAME_CHECKED(UMovieGraphExecuteScriptNode, EditorOnlyScript));
+		const TSharedRef<IPropertyHandle> EditorAndRuntimeHandle = DetailBuilder.GetProperty(GET_MEMBER_NAME_CHECKED(UMovieGraphExecuteScriptNode, EditorAndRuntimeScript));
+
+		// Gets the enum value of the Mode property.
+		auto GetModeValue = [ModeHandle]() -> EMovieGraphExecuteScriptMode
+		{
+			constexpr EMovieGraphExecuteScriptMode InvalidMode = EMovieGraphExecuteScriptMode::Count;
+
+			if (ModeHandle->IsValidHandle())
+			{
+				uint8 ModeValue = static_cast<uint8>(InvalidMode);
+
+				if (ModeHandle->GetValue(ModeValue) == FPropertyAccess::Success)
+				{
+					return static_cast<EMovieGraphExecuteScriptMode>(ModeValue);
+				}
+			}
+
+			return InvalidMode;
+		};
+
+		const TAttribute<EVisibility> EditorOnlyVisibleAttr = TAttribute<EVisibility>::Create([GetModeValue]()
+		{
+			return (GetModeValue() == EMovieGraphExecuteScriptMode::EditorOnly) ? EVisibility::Visible : EVisibility::Collapsed;
+		});
+
+		const TAttribute<EVisibility> EditorAndRuntimeVisibleAttr = TAttribute<EVisibility>::Create([GetModeValue]()
+		{
+			return (GetModeValue() == EMovieGraphExecuteScriptMode::EditorAndRuntime) ? EVisibility::Visible : EVisibility::Collapsed;
+		});
+
+		// Vary the visibility of the script properties depending on the selected mode.
+		if (IDetailPropertyRow* EditorOnlyRow = DetailBuilder.EditDefaultProperty(EditorOnlyHandle))
+		{
+			EditorOnlyRow->Visibility(EditorOnlyVisibleAttr);
+		}
+		if (IDetailPropertyRow* EditorAndRuntimeRow = DetailBuilder.EditDefaultProperty(EditorAndRuntimeHandle))
+		{
+			EditorAndRuntimeRow->Visibility(EditorAndRuntimeVisibleAttr);
+		}
+	}
+	//~ End IDetailCustomization interface
+};
+
+#undef LOCTEXT_NAMESPACE
\ No newline at end of file
diff --git a/Engine/Plugins/MovieScene/MovieRenderPipeline/Source/MovieRenderPipelineEditor/Private/Graph/MovieGraphAssetToolkit.cpp b/Engine/Plugins/MovieScene/MovieRenderPipeline/Source/MovieRenderPipelineEditor/Private/Graph/MovieGraphAssetToolkit.cpp
--- a/Engine/Plugins/MovieScene/MovieRenderPipeline/Source/MovieRenderPipelineEditor/Private/Graph/MovieGraphAssetToolkit.cpp
+++ b/Engine/Plugins/MovieScene/MovieRenderPipeline/Source/MovieRenderPipelineEditor/Private/Graph/MovieGraphAssetToolkit.cpp
@@ -4,6 +4,7 @@
 
 #include "Customizations/Graph/MovieGraphApplyCVarPresetNodeCustomization.h"
 #include "Customizations/Graph/MovieGraphCollectionsCustomization.h"
+#include "Customizations/Graph/MovieGraphExecuteScriptNodeCustomization.h"
 #include "Customizations/Graph/MovieGraphMultiLayerEXRNodeCustomization.h"
 #include "Customizations/Graph/MovieGraphMaterialParameterCollectionModifierNodeCustomization.h"
 #include "Customizations/Graph/MovieGraphMemberCustomization.h"
@@ -21,9 +22,7 @@
 #include "Graph/MovieGraphConfig.h"
 #include "Graph/Nodes/MovieGraphApplyCVarPresetNode.h"
 #include "Graph/Nodes/MovieGraphCollectionNode.h"
-#include "Graph/Nodes/MovieGraphCommandLineEncoderNode.h"
-#include "Graph/Nodes/MovieGraphDebugNode.h"
-#include "Graph/Nodes/MovieGraphFileOutputNode.h"
+#include "Graph/Nodes/MovieGraphExecuteScriptNode.h"
 #include "Graph/Nodes/MovieGraphMaterialParameterCollectionModifierNode.h"
 #include "Graph/Nodes/MovieGraphModifierNode.h"
 #include "Graph/Nodes/MovieGraphPathTracerPassNode.h"
@@ -444,6 +443,10 @@ TSharedRef<SDockTab> FMovieGraphAssetToolkit::SpawnTab_RenderGraphDetails(const
 	SelectedGraphObjectsDetailsWidget->RegisterInstancedCustomPropertyLayout(
 		UMovieGraphPathTracerRenderPassNode::StaticClass(),
 		FOnGetDetailCustomizationInstance::CreateStatic(&FMovieGraphPathTracedRendererNodeCustomization::MakeInstance));
+
+	SelectedGraphObjectsDetailsWidget->RegisterInstancedCustomPropertyLayout(
+		UMovieGraphExecuteScriptNode::StaticClass(),
+		FOnGetDetailCustomizationInstance::CreateStatic(&FMovieGraphExecuteScriptNodeCustomization::MakeInstance));
 	
 	TSharedRef<SWidget> CustomContent = SAssignNew(NameAreaCustomContent, SHorizontalBox)
 	+ SHorizontalBox::Slot()
