MoviePipeline: Fixed a crash that could occur if attempting to render with an unsaved map.

#Synced-CL 48243900

diff --git a/Engine/Plugins/MovieScene/MovieRenderPipeline/Source/MovieRenderPipelineEditor/Private/MovieRenderPipelineEditorUtils.cpp b/Engine/Plugins/MovieScene/MovieRenderPipeline/Source/MovieRenderPipelineEditor/Private/MovieRenderPipelineEditorUtils.cpp
--- a/Engine/Plugins/MovieScene/MovieRenderPipeline/Source/MovieRenderPipelineEditor/Private/MovieRenderPipelineEditorUtils.cpp
+++ b/Engine/Plugins/MovieScene/MovieRenderPipeline/Source/MovieRenderPipelineEditor/Private/MovieRenderPipelineEditorUtils.cpp
@@ -26,6 +26,12 @@ namespace UE::MovieRenderPipelineEditor::Private
 		
 		UMoviePipelineExecutorBase* Executor = Subsystem->RenderQueueWithExecutor(ExecutorClass);
 
+		// If the executor encountered an error (like trying to use an unsaved map), then the executor may be null at this point.
+		if (!IsValid(Executor))
+		{
+			return false;
+		}
+
 		// Tell the executor to run the post-render action. Some executors may not support this.
 		const UMoviePipelineQueueProjectSettings* QueueProjectSettings = GetDefault<UMoviePipelineQueueProjectSettings>();
 		Executor->SetPostRenderAction(QueueProjectSettings->PostRenderActionType);
