Movie Render Queue: Expand ability to choose submix to export

#Synced-CL 48822946

diff --git a/Engine/Plugins/MovieScene/MovieRenderPipeline/Source/MovieRenderPipelineCore/Private/MoviePipelineAudioRendering.cpp b/Engine/Plugins/MovieScene/MovieRenderPipeline/Source/MovieRenderPipelineCore/Private/MoviePipelineAudioRendering.cpp
--- a/Engine/Plugins/MovieScene/MovieRenderPipeline/Source/MovieRenderPipelineCore/Private/MoviePipelineAudioRendering.cpp
+++ b/Engine/Plugins/MovieScene/MovieRenderPipeline/Source/MovieRenderPipelineCore/Private/MoviePipelineAudioRendering.cpp
@@ -4,6 +4,7 @@
 #include "HAL/IConsoleManager.h"
 #include "AudioDeviceManager.h"
 #include "MovieRenderPipelineCoreModule.h"
+#include "MoviePipelineOutputSetting.h"
 #include "Engine/Engine.h"
 #include "AudioMixerBlueprintLibrary.h"
 #include "LevelSequence.h"
@@ -15,6 +16,7 @@
 #include "DSP/BufferVectorOperations.h"
 #include "MoviePipelineAntiAliasingSetting.h"
 #include "AudioMixerPlatformNonRealtime.h"
+
 #include "MoviePipelineQueue.h"
 
 static FAudioDevice* GetAudioDeviceFromWorldContext(const UObject* WorldContextObject)
@@ -71,6 +73,9 @@ void UMoviePipeline::TeardownAudioRendering()
 void UMoviePipeline::StartAudioRecording()
 {
 	AudioState.bIsRecordingAudio = true;
+	const UMoviePipelineExecutorShot* CurrentShot = ActiveShotList[CurrentShotIndex];
+	UMoviePipelineOutputSetting* OutputSettings = FindOrAddSettingForShot<UMoviePipelineOutputSetting>(CurrentShot);
+	check(OutputSettings);
 
 	Audio::FMixerDevice* MixerDevice = GetAudioMixerDeviceFromWorldContext(this);
 	if (MixerDevice)
@@ -81,6 +86,15 @@ void UMoviePipeline::StartAudioRecording()
 			AudioState.ActiveSubmixes.Add(MasterSubmix);
 		}
 
+		for (USoundSubmix* Mix : OutputSettings->AdditionalSubmixOutputs)
+		{
+			TWeakPtr<Audio::FMixerSubmix, ESPMode::ThreadSafe> Submix = MixerDevice->GetSubmixInstance(Mix);
+			if (Submix.Pin())
+			{
+				AudioState.ActiveSubmixes.Add(Submix);
+			}
+		}
+
 		for (TWeakPtr<Audio::FMixerSubmix, ESPMode::ThreadSafe> WeakSubmix : AudioState.ActiveSubmixes)
 		{
 			const float ExpectedDuration = 30.f;
@@ -104,6 +118,7 @@ void UMoviePipeline::StopAudioRecording()
 			// Copy the data returned by the Submix.
 			Audio::AlignedFloatBuffer& AudioRecording = Submix->OnStopRecordingOutput(NewSegment.NumChannels, NewSegment.SampleRate);
 			NewSegment.SegmentData = AudioRecording; 
+			NewSegment.SubmixName = Submix->GetName();
 		}
 	}
 
diff --git a/Engine/Plugins/MovieScene/MovieRenderPipeline/Source/MovieRenderPipelineCore/Public/MoviePipelineOutputSetting.h b/Engine/Plugins/MovieScene/MovieRenderPipeline/Source/MovieRenderPipelineCore/Public/MoviePipelineOutputSetting.h
--- a/Engine/Plugins/MovieScene/MovieRenderPipeline/Source/MovieRenderPipelineCore/Public/MoviePipelineOutputSetting.h
+++ b/Engine/Plugins/MovieScene/MovieRenderPipeline/Source/MovieRenderPipelineCore/Public/MoviePipelineOutputSetting.h
@@ -5,6 +5,7 @@
 #include "Engine/EngineTypes.h"
 #include "Misc/FrameRate.h"
 #include "MoviePipelineColorSetting.h"
+#include "Sound/SoundSubmix.h"
 #include "MoviePipelineOutputSetting.generated.h"
 
 #define UE_API MOVIERENDERPIPELINECORE_API
@@ -97,6 +98,10 @@ public:
 	UPROPERTY(EditAnywhere, BlueprintReadWrite, Category = "Versioning")
 	bool bAutoVersion;
 
+	/** Exports additional .wav file for each submix if .wav output is selected*/
+	UPROPERTY(EditAnywhere, BlueprintReadWrite, Category = "File Output")
+	TArray<TObjectPtr<USoundSubmix>> AdditionalSubmixOutputs;
+
 public:
 	/** How many digits should all output frame numbers be padded to? MySequence_1.png -> MySequence_0001.png. Useful for software that struggles to recognize frame ranges when non-padded. */
 	UPROPERTY(EditAnywhere, BlueprintReadWrite, meta = (UIMin = "1", UIMax = "5", ClampMin = "1"), AdvancedDisplay, Category = "File Output")
diff --git a/Engine/Plugins/MovieScene/MovieRenderPipeline/Source/MovieRenderPipelineCore/Public/MovieRenderPipelineDataTypes.h b/Engine/Plugins/MovieScene/MovieRenderPipeline/Source/MovieRenderPipelineCore/Public/MovieRenderPipelineDataTypes.h
--- a/Engine/Plugins/MovieScene/MovieRenderPipeline/Source/MovieRenderPipelineCore/Public/MovieRenderPipelineDataTypes.h
+++ b/Engine/Plugins/MovieScene/MovieRenderPipeline/Source/MovieRenderPipelineCore/Public/MovieRenderPipelineDataTypes.h
@@ -1318,6 +1318,8 @@ namespace MoviePipeline
 			FGuid Id;
 			FMoviePipelineFrameOutputState OutputState;
 
+			FString SubmixName;
+
 			float NumChannels;
 			float SampleRate;
 
diff --git a/Engine/Plugins/MovieScene/MovieRenderPipeline/Source/MovieRenderPipelineRenderPasses/Private/MoviePipelineWaveOutput.cpp b/Engine/Plugins/MovieScene/MovieRenderPipeline/Source/MovieRenderPipelineRenderPasses/Private/MoviePipelineWaveOutput.cpp
--- a/Engine/Plugins/MovieScene/MovieRenderPipeline/Source/MovieRenderPipelineRenderPasses/Private/MoviePipelineWaveOutput.cpp
+++ b/Engine/Plugins/MovieScene/MovieRenderPipeline/Source/MovieRenderPipelineRenderPasses/Private/MoviePipelineWaveOutput.cpp
@@ -115,9 +115,11 @@ void UMoviePipelineWaveOutput::BeginFinalizeImpl()
 		FormatOverrides.Add(TEXT("render_pass"), TEXT("Audio"));
 		FormatOverrides.Add(TEXT("ext"), TEXT("wav"));
 
+		// Append the submix name to have a different .wav output per submix
+		FString FileNameWithSubmix = FileNameFormatString + Segment.SubmixName;
 		FMoviePipelineFormatArgs FinalFormatArgs;
 		FString FinalFilePath;
-		GetPipeline()->ResolveFilenameFormatArguments(FileNameFormatString, FormatOverrides, FinalFilePath, FinalFormatArgs, &Segment.OutputState);
+		GetPipeline()->ResolveFilenameFormatArguments(FileNameWithSubmix, FormatOverrides, FinalFilePath, FinalFormatArgs, &Segment.OutputState);
 
 		if (FPaths::IsRelative(FinalFilePath))
 		{
diff --git a/Engine/Plugins/Runtime/AudioModulation/Source/AudioModulation/AudioModulation.Build.cs b/Engine/Plugins/Runtime/AudioModulation/Source/AudioModulation/AudioModulation.Build.cs
--- a/Engine/Plugins/Runtime/AudioModulation/Source/AudioModulation/AudioModulation.Build.cs
+++ b/Engine/Plugins/Runtime/AudioModulation/Source/AudioModulation/AudioModulation.Build.cs
@@ -19,7 +19,8 @@ namespace UnrealBuildTool.Rules
 					"Core",
 					"WaveTable",
 					"MovieSceneTracks",
-					"MovieScene"
+					"MovieScene",
+					"AudioMixer"
 				}
 			);
 
