MoviePipeline: Added a new Light Modifier node to the graph. This modifier contains a set of built-in light properties that are commonly-modified, but also allows for nearly any light property to be added via the "Custom" section.

Driving the Light Modifier node is a new generic actor/component modifier (MovieGraphClassPropertyModifier). The goal of this modifier is to make simple actor/component modification reusable and easy.

Several API-related changes were made. Notably, MovieGraphValueView was introduced, which provides a scripting-friendly way of modifying a node's dynamic properties (which drive the "Custom" properties on the light modifier) or any other property bag. Additionally, MovieGraphMultiValueContainer was added, which provides a simple and scripting-friendly wrapper around a property bag which can contain multiple values (vs MovieGraphValueContainer, which can only contain one value).

Several modifier and collection-related widgets were also separated out and made public, so any modifier nodes that need to take advantage of built-in modifier-related widgetry can.

#Synced-CL 49282327

diff --git a/Engine/Plugins/MovieScene/MovieRenderPipeline/Source/MovieRenderPipelineCore/Private/Graph/MovieGraphNode.cpp b/Engine/Plugins/MovieScene/MovieRenderPipeline/Source/MovieRenderPipelineCore/Private/Graph/MovieGraphNode.cpp
--- a/Engine/Plugins/MovieScene/MovieRenderPipeline/Source/MovieRenderPipelineCore/Private/Graph/MovieGraphNode.cpp
+++ b/Engine/Plugins/MovieScene/MovieRenderPipeline/Source/MovieRenderPipelineCore/Private/Graph/MovieGraphNode.cpp
@@ -100,6 +100,11 @@ namespace UE::MovieGraph::Private
 
 UMovieGraphNode::UMovieGraphNode()
 {
+	if (!HasAnyFlags(RF_ClassDefaultObject))
+	{
+		DynamicPropertiesView = NewObject<UMovieGraphMutableValueView>(this, TEXT("DynamicPropertiesView"));
+		DynamicPropertiesView->ViewProperties(&DynamicProperties);
+	}
 }
 
 TArray<FMovieGraphPinProperties> UMovieGraphNode::GetExposedPinProperties() const
@@ -493,6 +498,11 @@ bool UMovieGraphNode::GetDynamicPropertyValue(const FName PropertyName, FString&
 	return Result.IsValid();
 }
 
+UMovieGraphMutableValueView* UMovieGraphNode::GetDynamicPropertiesView() const
+{
+	return DynamicPropertiesView.Get();
+}
+
 const FBoolProperty* UMovieGraphNode::FindOverridePropertyForDynamicProperty(const FName& InPropertyName) const
 {
 	const FString OverridePropertyName = FString::Printf(TEXT("bOverride_%s"), *InPropertyName.ToString());
diff --git a/Engine/Plugins/MovieScene/MovieRenderPipeline/Source/MovieRenderPipelineCore/Private/Graph/MovieGraphValueUtils.cpp b/Engine/Plugins/MovieScene/MovieRenderPipeline/Source/MovieRenderPipelineCore/Private/Graph/MovieGraphValueUtils.cpp
new file mode 100644
--- /dev/null
+++ b/Engine/Plugins/MovieScene/MovieRenderPipeline/Source/MovieRenderPipelineCore/Private/Graph/MovieGraphValueUtils.cpp
@@ -0,0 +1,121 @@
+﻿// Copyright Epic Games, Inc. All Rights Reserved.
+
+#include "Graph/MovieGraphValueUtils.h"
+
+#include "UObject/EnumProperty.h"
+#include "UObject/StrProperty.h"
+#include "UObject/TextProperty.h"
+#include "UObject/UnrealType.h"
+
+bool UE::MovieGraph::ValueUtils::GetTypesForFProperty(
+	const FProperty* InProperty, EMovieGraphValueType& OutValueType, TObjectPtr<UObject>& OutValueTypeObject)
+{
+	OutValueTypeObject = nullptr;
+	OutValueType = EMovieGraphValueType::None;
+
+	if (!InProperty)
+	{
+		return false;
+	}
+
+	if (InProperty->IsA<FBoolProperty>())
+	{
+		OutValueType = EMovieGraphValueType::Bool;
+	}
+	else if (InProperty->IsA<FByteProperty>())
+	{
+		OutValueType = EMovieGraphValueType::Byte;
+
+		if (const TObjectPtr<UEnum> ByteAsEnum = CastField<FByteProperty>(InProperty)->Enum)
+		{
+			OutValueType = EMovieGraphValueType::Enum;
+			OutValueTypeObject = ByteAsEnum;
+		}
+	}
+	else if (InProperty->IsA<FIntProperty>())
+	{
+		OutValueType = EMovieGraphValueType::Int32;
+	}
+	else if (InProperty->IsA<FInt64Property>())
+	{
+		OutValueType = EMovieGraphValueType::Int64;
+	}
+	else if (InProperty->IsA<FFloatProperty>())
+	{
+		OutValueType = EMovieGraphValueType::Float;
+	}
+	else if (InProperty->IsA<FDoubleProperty>())
+	{
+		OutValueType = EMovieGraphValueType::Double;
+	}
+	else if (InProperty->IsA<FNameProperty>())
+	{
+		OutValueType = EMovieGraphValueType::Name;
+	}
+	else if (InProperty->IsA<FStrProperty>())
+	{
+		OutValueType = EMovieGraphValueType::String;
+	}
+	else if (InProperty->IsA<FTextProperty>())
+	{
+		OutValueType = EMovieGraphValueType::Text;
+	}
+	else if (InProperty->IsA<FEnumProperty>())
+	{
+		OutValueType = EMovieGraphValueType::Enum;
+		OutValueTypeObject = CastField<FEnumProperty>(InProperty)->GetEnum();
+	}
+	else if (InProperty->IsA<FStructProperty>())
+	{
+		OutValueType = EMovieGraphValueType::Struct;
+		OutValueTypeObject = CastField<FStructProperty>(InProperty)->Struct.Get();
+	}
+	else if (InProperty->IsA<FObjectPropertyBase>())
+	{
+		UClass* PropertyClass = CastField<FObjectPropertyBase>(InProperty)->PropertyClass.Get();
+		
+		if (InProperty->IsA<FObjectProperty>())
+		{
+			if (PropertyClass == UClass::StaticClass())
+			{
+				OutValueType = EMovieGraphValueType::Class;
+				OutValueTypeObject = UObject::StaticClass();
+			}
+			else
+			{
+				OutValueType = EMovieGraphValueType::Object;
+				OutValueTypeObject = PropertyClass;
+			}
+		}
+		else if (InProperty->IsA<FSoftObjectProperty>())
+		{
+			OutValueType = EMovieGraphValueType::SoftObject;
+			OutValueTypeObject = PropertyClass;
+		}
+		else if (InProperty->IsA<FClassProperty>())
+		{
+			OutValueType = EMovieGraphValueType::Class;
+			OutValueTypeObject = CastField<FClassProperty>(InProperty)->MetaClass;
+		}
+		else if (InProperty->IsA<FSoftClassProperty>())
+		{
+			OutValueType = EMovieGraphValueType::SoftClass;
+			OutValueTypeObject = CastField<FSoftClassProperty>(InProperty)->MetaClass;
+		}
+	}
+	else if (InProperty->IsA<FUInt32Property>())
+	{
+		OutValueType = EMovieGraphValueType::UInt32;
+	}
+	else if (InProperty->IsA<FUInt64Property>())
+	{
+		OutValueType = EMovieGraphValueType::UInt64;
+	}
+	else
+	{
+		// Couldn't find the value type and/or value type object for the property.
+		return false;
+	}
+
+	return true;
+}
diff --git a/Engine/Plugins/MovieScene/MovieRenderPipeline/Source/MovieRenderPipelineCore/Private/Graph/Nodes/MovieGraphClassPropertyModifier.cpp b/Engine/Plugins/MovieScene/MovieRenderPipeline/Source/MovieRenderPipelineCore/Private/Graph/Nodes/MovieGraphClassPropertyModifier.cpp
new file mode 100644
--- /dev/null
+++ b/Engine/Plugins/MovieScene/MovieRenderPipeline/Source/MovieRenderPipelineCore/Private/Graph/Nodes/MovieGraphClassPropertyModifier.cpp
@@ -0,0 +1,287 @@
+// Copyright Epic Games, Inc. All Rights Reserved.
+
+#include "Graph/Nodes/MovieGraphClassPropertyModifier.h"
+
+#include "Engine/World.h"
+#include "Graph/MovieGraphValueUtils.h"
+#include "Graph/Nodes/MovieGraphLightModifierNode.h"
+#include "Kismet/GameplayStatics.h"
+#include "MovieRenderPipelineCoreModule.h"
+
+FMovieGraphPropertyReference::FMovieGraphPropertyReference(
+	const TObjectPtr<UClass>& InPropertyClass,
+	const TSubclassOf<UActorComponent>& InComponentClass,
+	const FName& InPropertyName,
+	FProperty* InProperty)
+	: PropertyClass(InPropertyClass)
+	, ComponentClass(InComponentClass)
+	, PropertyName(InPropertyName)
+	, Property(InProperty)
+{
+	PopulateDerivedData();
+}
+
+EMovieGraphValueType FMovieGraphPropertyReference::GetValueType() const
+{
+	return ValueType;
+}
+
+UObject* FMovieGraphPropertyReference::GetValueTypeObject() const
+{
+	return ValueTypeObject.Get();
+}
+
+FProperty* FMovieGraphPropertyReference::GetProperty() const
+{
+	return Property;
+}
+
+FName FMovieGraphPropertyReference::GetPropertyName() const
+{
+	return PropertyName;
+}
+
+TSubclassOf<AActor> FMovieGraphPropertyReference::GetPropertyClass() const
+{
+	return PropertyClass;
+}
+
+TSubclassOf<UActorComponent> FMovieGraphPropertyReference::GetComponentClass() const
+{
+	return ComponentClass;
+}
+
+bool FMovieGraphPropertyReference::IsValidProperty() const
+{
+	return !!Property && PropertyClass.Get() && !PropertyName.IsNone();
+}
+
+void FMovieGraphPropertyReference::PostSerialize(const FArchive& Ar)
+{
+	if (Ar.IsLoading())
+	{
+		// FProperty cannot be serialized, so we re-derive it upon loading (and any other calculated members)
+		PopulateDerivedData();
+	}
+}
+
+void FMovieGraphPropertyReference::PopulateDerivedData()
+{
+	// If the property is part of a component, use the component as the container property; otherwise, use the class.
+	if (!Property)
+	{
+		if (ComponentClass.Get())
+		{
+			Property = FindFProperty<FProperty>(ComponentClass.Get(), PropertyName);
+		}
+		else if (PropertyClass.Get())
+		{
+			Property = FindFProperty<FProperty>(PropertyClass.Get(), PropertyName);
+		}
+	}
+
+	if (!Property)
+	{
+		UE_LOG(LogMovieRenderPipeline, Warning, TEXT("Could not find property with name [%s]. Functionality relating to this property may be broken."), *PropertyName.ToString());
+	}
+
+	if (Property)
+	{
+		UE::MovieGraph::ValueUtils::GetTypesForFProperty(Property, ValueType, ValueTypeObject);
+	}
+}
+
+UMovieGraphClassPropertyModifier::UMovieGraphClassPropertyModifier()
+{
+	if (!HasAnyFlags(RF_ClassDefaultObject))
+	{
+		UpdatedPropertyValuesView = NewObject<UMovieGraphMutableValueView>(this, TEXT("UpdatedPropertyValuesView"));
+		UpdatedPropertyValuesView->ViewProperties(&UpdatedPropertyValues);
+
+		PreviousPropertyValuesView = NewObject<UMovieGraphMutableValueView>(this, TEXT("PreviousPropertyValuesView"));
+		PreviousPropertyValuesView->ViewProperties(&PreviousPropertyValues);
+	}
+}
+
+void UMovieGraphClassPropertyModifier::ApplyModifier(const UWorld* World)
+{
+	TArray<TSubclassOf<AActor>> TargetActorClasses;
+	PropertiesToModifyByActorClass.GetKeys(TargetActorClasses);
+
+	// Keep track of the components that were updated; they need their render state marked as dirty so the render thread gets the update(s).
+	TArray<UActorComponent*> DirtiedComponents;
+
+	for (const TObjectPtr<UMovieGraphCollection>& Collection : Collections)
+	{
+		if (!Collection)
+		{
+			continue;
+		}
+
+		FMovieGraphEvaluationResult EvaluationResult = Collection->EvaluateActorsAndComponents(World);
+
+		for (const TObjectPtr<AActor>& MatchingActor : EvaluationResult.MatchingActors)
+		{
+			UClass* ActorClass = MatchingActor->GetClass();
+
+			// The modifier only applies to specific actor types
+			if (!TargetActorClasses.Contains(ActorClass))
+			{
+				continue;
+			}
+
+			for (const FMovieGraphPropertyReference& PropertyReference : PropertiesToModifyByActorClass[ActorClass])
+			{
+				// If the property is part of a component, use the component as the container property; otherwise, use the actor itself.
+				UObject* PropertyContainer = nullptr;
+				if (PropertyReference.GetComponentClass())
+				{
+					PropertyContainer = MatchingActor->GetComponentByClass(PropertyReference.GetComponentClass());
+					DirtiedComponents.Add(Cast<UActorComponent>(PropertyContainer));
+				}
+				else
+				{
+					PropertyContainer = MatchingActor;
+				}
+
+				if (!PropertyContainer)
+				{
+					continue;
+				}
+
+				// TODO: Ideally the below export/import aren't string based, but that would require a lot of specialized code in order to
+				// properly handle every type. This will do for now, but could potentially be optimized in the future if this is a bottleneck
+				// or causes issues.
+
+				// Cache the current value so it can be restored later
+				FString CurrentValue;
+				PropertyReference.GetProperty()->ExportTextItem_InContainer(CurrentValue, PropertyContainer, nullptr, nullptr, PPF_None);
+				PreviousPropertyValuesView->SetValueSerializedString(PropertyReference.AuxPropertyName, CurrentValue);
+				ModifiedProperties.Add({ PropertyContainer, PropertyReference });
+				
+				// Set the property to its new value
+				const FString PropertyValue = UpdatedPropertyValuesView->GetValueSerializedString(PropertyReference.AuxPropertyName);
+				PropertyReference.GetProperty()->ImportText_InContainer(*PropertyValue, PropertyContainer, PropertyContainer, 0);
+			}
+		}
+	}
+
+	// The render thread needs the updates that were made here.
+	for (UActorComponent* DirtyComponent : DirtiedComponents)
+	{
+		if (DirtyComponent)
+		{
+			DirtyComponent->MarkRenderStateDirty();
+		}
+	}
+}
+
+void UMovieGraphClassPropertyModifier::UndoModifier()
+{
+	for (const TPair<UObject*, FMovieGraphPropertyReference>& ModifiedPropertyPair : ModifiedProperties)
+	{
+		UObject* ModifiedObject = ModifiedPropertyPair.Key;
+		const FMovieGraphPropertyReference& ModifiedPropertyReference = ModifiedPropertyPair.Value;
+		const FProperty* ModifiedProperty = ModifiedPropertyPair.Value.GetProperty();
+
+		if (ModifiedProperty && ModifiedObject)
+		{
+			// Restore the property value to what it was prior to being updated by the modifier
+			const FString PreviousValue = PreviousPropertyValuesView->GetValueSerializedString(ModifiedPropertyReference.AuxPropertyName);
+			ModifiedProperty->ImportText_InContainer(*PreviousValue, ModifiedObject, ModifiedObject, 0);
+
+			// Ensure the render thread gets this update
+			if (UActorComponent* ActorComponent = Cast<UActorComponent>(ModifiedObject))
+			{
+				ActorComponent->MarkRenderStateDirty();
+			}
+		}
+	}
+
+	ModifiedProperties.Reset();
+}
+
+FText UMovieGraphClassPropertyModifier::GetModifierName()
+{
+	static const FText ModifierName = NSLOCTEXT("MovieGraph", "ClassPropertyModifierName", "Class Property");
+	return ModifierName;
+}
+
+FMovieGraphPropertyReference UMovieGraphClassPropertyModifier::MakePropertyReference(
+	UClass* InPropertyClass,
+	const TSubclassOf<UActorComponent>& InComponentClass,
+	const FName& InPropertyName)
+{
+	return FMovieGraphPropertyReference(InPropertyClass, InComponentClass, InPropertyName);
+}
+
+void UMovieGraphClassPropertyModifier::AddProperties(const TArray<FMovieGraphPropertyReference>& InProperties)
+{
+	TArray<FMovieGraphValueViewProperty> PropertiesToAdd;
+
+	PropertiesToAdd.Reserve(InProperties.Num());
+
+	// Intentionally iterated by value instead of reference. The AuxPropertyName needs to be added to the reference, and the reference will be
+	// copied anyway when added to the PropertiesToModify array.
+	for (FMovieGraphPropertyReference PropertyReference : InProperties)
+	{
+		if (!PropertyReference.IsValidProperty())
+		{
+			continue;
+		}
+
+		// Populate the AuxPropertyName; this is the "internal" name of the property within the modifier which is used in the property bags.
+		PropertyReference.AuxPropertyName = GetInternalPropertyName(PropertyReference);
+
+		// Batch up this property to be added to the internal property bags.
+		PropertiesToAdd.Emplace(FMovieGraphValueViewProperty(
+			PropertyReference.AuxPropertyName, PropertyReference.GetValueType(), PropertyReference.GetValueTypeObject(), {}));
+
+		PropertiesToModifyByActorClass.FindOrAdd(PropertyReference.GetPropertyClass()).Add(PropertyReference);
+	}
+
+	// With all the data about the properties found, add all of them in one batch for efficiency.
+	UpdatedPropertyValuesView->AddProperties(PropertiesToAdd);
+	PreviousPropertyValuesView->AddProperties(PropertiesToAdd);
+}
+
+FName UMovieGraphClassPropertyModifier::GetPropertyNameForSettingCustomValue(const FMovieGraphPropertyReference& InPropertyReference)
+{
+	return GetInternalPropertyName(InPropertyReference);
+}
+
+UMovieGraphFixedValueView* UMovieGraphClassPropertyModifier::GetPropertiesForSettingValues()
+{
+	if (!UserFacingPropertiesView)
+	{
+		UserFacingPropertiesView = NewObject<UMovieGraphFixedValueView>(this, TEXT("UserFacingPropertiesView"));
+		UserFacingPropertiesView->ViewProperties(&UpdatedPropertyValues);
+	}
+
+	return UserFacingPropertiesView;
+}
+
+FName UMovieGraphClassPropertyModifier::GetInternalPropertyName(const FMovieGraphPropertyReference& InPropertyReference)
+{
+	// If the property reference includes a component, that needs to be included in the property name in order to properly disambiguate it.
+	// Otherwise, using the class name is sufficient.
+
+	if (InPropertyReference.GetComponentClass().Get() && InPropertyReference.GetPropertyClass())
+	{
+		return FName(FString::Format(TEXT("{0}_{1}_{2}"), {
+			InPropertyReference.GetPropertyClass()->GetName(),
+			InPropertyReference.GetComponentClass()->GetName(),
+			InPropertyReference.GetPropertyName().ToString() }));
+	}
+
+	if (InPropertyReference.GetPropertyClass())
+	{
+		return FName(FString::Format(TEXT("{0}_{1}"), {
+			InPropertyReference.GetPropertyClass()->GetName(),
+			InPropertyReference.GetPropertyName().ToString() }));
+	}
+
+	UE_LOG(LogMovieRenderPipeline, Warning, TEXT("Attempting to modify a property via the Property Modifier with an invalid actor class and/or component class."));
+
+	return FName(TEXT("INVALID"));
+}
diff --git a/Engine/Plugins/MovieScene/MovieRenderPipeline/Source/MovieRenderPipelineCore/Private/Graph/Nodes/MovieGraphLightModifierNode.cpp b/Engine/Plugins/MovieScene/MovieRenderPipeline/Source/MovieRenderPipelineCore/Private/Graph/Nodes/MovieGraphLightModifierNode.cpp
new file mode 100644
--- /dev/null
+++ b/Engine/Plugins/MovieScene/MovieRenderPipeline/Source/MovieRenderPipelineCore/Private/Graph/Nodes/MovieGraphLightModifierNode.cpp
@@ -0,0 +1,511 @@
+// Copyright Epic Games, Inc. All Rights Reserved.
+
+#include "Graph/Nodes/MovieGraphLightModifierNode.h"
+
+#include "Components/LightComponent.h"
+#include "Components/LightComponentBase.h"
+#include "Components/LocalLightComponent.h"
+#include "Engine/DirectionalLight.h"
+#include "Engine/PointLight.h"
+#include "Engine/RectLight.h"
+#include "Engine/SkyLight.h"
+#include "Engine/SpotLight.h"
+#include "GameFramework/Actor.h"
+#include "Graph/MovieGraphValueUtils.h"
+#include "Graph/Nodes/MovieGraphClassPropertyModifier.h"
+
+#include UE_INLINE_GENERATED_CPP_BY_NAME(MovieGraphLightModifierNode)
+
+#define LOCTEXT_NAMESPACE "MovieGraph"
+
+UMovieGraphLightModifierNode::UMovieGraphLightModifierNode()
+{
+	Modifier = CreateDefaultSubobject<UMovieGraphClassPropertyModifier>(TEXT("LightingModifier"));
+}
+
+#if WITH_EDITOR
+FText UMovieGraphLightModifierNode::GetNodeTitle(const bool bGetDescriptive) const
+{
+	static const FText NodeName = LOCTEXT("NodeName_LightModifier", "Light Modifier");
+	static const FText NodeDescription = LOCTEXT("NodeDescription_LightModifier", "Light Modifier\n{0}");
+
+	const FString ModifierNameDisp = ModifierName.IsEmpty() ? LOCTEXT("NodeNoNameWarning_LightModifier", "NO NAME").ToString() : ModifierName;
+
+	if (bGetDescriptive)
+	{
+		return FText::Format(NodeDescription, FText::FromString(ModifierNameDisp));
+	}
+
+	return NodeName;
+}
+
+FText UMovieGraphLightModifierNode::GetMenuCategory() const
+{
+	return LOCTEXT("LightModifierNode_Category", "Utility");
+}
+
+FLinearColor UMovieGraphLightModifierNode::GetNodeTitleColor() const
+{
+	static const FLinearColor ModifierNodeColor = FLinearColor(0.6f, 0.113f, 0.113f);
+	return ModifierNodeColor;
+}
+
+FSlateIcon UMovieGraphLightModifierNode::GetIconAndTint(FLinearColor& OutColor) const
+{
+	static const FSlateIcon ModifierIcon = FSlateIcon(FAppStyle::GetAppStyleSetName(), "ClassIcon.PointLightComponent");
+
+	OutColor = FLinearColor::White;
+	return ModifierIcon;
+}
+
+void UMovieGraphLightModifierNode::PostEditChangeProperty(FPropertyChangedEvent& PropertyChangedEvent)
+{
+	Super::PostEditChangeProperty(PropertyChangedEvent);
+
+	// Broadcast a node-changed delegate so that the node title's UI gets updated.
+	if ((PropertyChangedEvent.GetPropertyName() == GET_MEMBER_NAME_CHECKED(UMovieGraphModifierNode, ModifierName)) ||
+		(PropertyChangedEvent.GetPropertyName() == GET_MEMBER_NAME_CHECKED(UMovieGraphModifierNode, bOverride_ModifierName)))
+	{
+		OnNodeChangedDelegate.Broadcast(this);
+	}
+}
+#endif // WITH_EDITOR
+
+TArray<FPropertyBagPropertyDesc> UMovieGraphLightModifierNode::GetDynamicPropertyDescriptions() const
+{
+	// Most of these should really be defined somewhere outside of MRG, but they don't appear to be.
+	static const FName EditConditionMetadataKey = TEXT("EditCondition");
+	static const FName CategoryMetadataKey = TEXT("Category");
+	static const FName EnableCategoriesMetadataKey = TEXT("EnableCategories");
+	static const FName LightingClassMetadataKey = TEXT("LightingClass");
+	static const FName DisplayNameMetadataKey = TEXT("DisplayName");
+	static const FName InlineEditConditionToggleMetadataKey = TEXT("InlineEditConditionToggle");
+
+	TArray<FPropertyBagPropertyDesc> PropertyDescs;
+
+	for (const TPair<FName, FMovieGraphPropertyReference>& CustomPropertyPair : CustomPropertyInfo)
+	{
+		const FMovieGraphPropertyReference& CustomPropertyRef = CustomPropertyPair.Value;
+		const FProperty* CustomProperty = CustomPropertyRef.GetProperty();
+		const FName OverridePropertyName = GetLightingOverridePropertyName(CustomPropertyRef.AuxPropertyName);
+
+		EMovieGraphValueType ValueType = EMovieGraphValueType::None;
+		TObjectPtr<UObject> ValueTypeObject = nullptr;
+		if (!UE::MovieGraph::ValueUtils::GetTypesForFProperty(CustomProperty, ValueType, ValueTypeObject))
+		{
+			continue;
+		}
+
+		FPropertyBagPropertyDesc MainDesc(CustomPropertyRef.AuxPropertyName, static_cast<EPropertyBagPropertyType>(ValueType), ValueTypeObject);
+
+#if WITH_EDITOR
+		// Use all of the metadata associated with the property
+		for (const TPair<FName, FString>& MetaDataPair : *CustomProperty->GetMetaDataMap())
+		{
+			MainDesc.MetaData.Add({ MetaDataPair.Key, MetaDataPair.Value });
+		}
+
+		// Enable categories so that they're grouped by light type. The category also needs to be changed to the class associated with the property.
+		if (const TSubclassOf<AActor> PropertyClass = CustomPropertyRef.GetPropertyClass())
+		{
+			const FString LightingClassDisplayName = PropertyClass->GetDisplayNameText().ToString();
+
+			MainDesc.MetaData.Add({ CategoryMetadataKey, LightingClassDisplayName });
+			MainDesc.MetaData.Add({ EnableCategoriesMetadataKey, FString(TEXT("true")) });
+
+			// Parts of the UI may need to access the name of the lighting class; this is not officially sanctioned metadata and is MRG-specific.
+			MainDesc.MetaData.Add({ LightingClassMetadataKey, LightingClassDisplayName });
+		}
+
+		// Multiple light classes could use the same property, so the display name needs to be set (since the property name is prefixed by the light class).
+		MainDesc.MetaData.Add({ DisplayNameMetadataKey, CustomProperty->GetDisplayNameText().ToString() });
+
+		// Like other MRG properties, this can be selectively enabled/disabled via an override property.
+		MainDesc.MetaData.Add({ EditConditionMetadataKey, OverridePropertyName.ToString() });
+
+		// Some lighting properties may be used as inline Edit Condition toggles. For the purposes of Custom properties, we *don't* want it behaving that
+		// way, since MRG adds its own override toggle for properties.
+		MainDesc.MetaData.RemoveAll([](const FPropertyBagPropertyDescMetaData& MetaDataEntry)
+		{
+			return MetaDataEntry.Key == InlineEditConditionToggleMetadataKey;
+		});
+#endif	// WITH_EDITOR
+
+		// Add the main property
+		PropertyDescs.Add(MainDesc);
+
+		// Add the override property (EditCondition checkbox next to the property).
+		FPropertyBagPropertyDesc OverrideDesc(OverridePropertyName, EPropertyBagPropertyType::Bool);
+#if WITH_EDITOR
+		OverrideDesc.MetaData.Add({ InlineEditConditionToggleMetadataKey, TEXT("true") });
+#endif
+		PropertyDescs.Add(MoveTemp(OverrideDesc));
+	}
+
+	return PropertyDescs;
+}
+
+TArray<FMovieGraphPropertyInfo> UMovieGraphLightModifierNode::GetOverrideablePropertyInfo() const
+{
+	TArray<FMovieGraphPropertyInfo> PropertyInfo = Super::GetOverrideablePropertyInfo();
+
+	// Exclude the "Dynamic Properties" property from being shown. All of the properties *within* that property bag will be
+	// shown instead.
+	PropertyInfo.RemoveAll([](const FMovieGraphPropertyInfo& InPropertyInfo)
+	{
+		// GET_MEMBER_NAME_CHECKED can't be used because DynamicProperties is protected
+		return InPropertyInfo.Name == TEXT("DynamicProperties");
+	});
+
+	const UMovieGraphMutableValueView* LightingValues = GetDynamicPropertiesView();
+
+	// The PropertyInfo array will include all of the "Custom" properties since it already takes into account the DynamicProperties property bag.
+	// However, these entries should have better display names in the context menu.
+	for (FMovieGraphPropertyInfo& Info : PropertyInfo)
+	{
+		if (!Info.bIsDynamicProperty)
+		{
+			continue;
+		}
+
+#if WITH_EDITOR
+		const TMap<FName, FString> PropertyMetadata = LightingValues->GetPropertyMetadata(Info.Name);
+		const FString* DisplayNameMetadata = PropertyMetadata.Find(TEXT("DisplayName"));
+		const FString* LightingClassMetadata = PropertyMetadata.Find(TEXT("LightingClass"));
+		if (DisplayNameMetadata && LightingClassMetadata)
+		{
+			Info.ContextMenuName =
+				FText::FromString(FString::Format(TEXT("[{0}] {1}"), { *LightingClassMetadata, *DisplayNameMetadata }));
+		}
+#endif	// WITH_EDITOR
+	}
+
+	return PropertyInfo;
+}
+
+void UMovieGraphLightModifierNode::PrepareForFlattening(const UMovieGraphSettingNode* InSourceNode)
+{
+	Super::PrepareForFlattening(InSourceNode);
+
+	// This node's dynamic properties rely on the CustomPropertyInfo being present (since GetDynamicPropertyDescriptions() iterates it)
+	CustomPropertyInfo = Cast<UMovieGraphLightModifierNode>(InSourceNode)->CustomPropertyInfo;
+}
+
+TArray<UMovieGraphModifierBase*> UMovieGraphLightModifierNode::GetAllModifiers() const
+{
+	// Lighting component property names
+	static const FName LightColorName = GET_MEMBER_NAME_CHECKED(ULightComponentBase, LightColor);
+	static const FName AffectsWorldName = GET_MEMBER_NAME_CHECKED(ULightComponentBase, bAffectsWorld);
+	static const FName CastShadowsName = GET_MEMBER_NAME_CHECKED(ULightComponentBase, CastShadows);
+	static const FName IndirectLightingIntensityName = GET_MEMBER_NAME_CHECKED(ULightComponentBase, IndirectLightingIntensity);
+	static const FName VolumetricScatteringIntensityName = GET_MEMBER_NAME_CHECKED(ULightComponentBase, VolumetricScatteringIntensity);
+	static const FName LightingChannelsName = GET_MEMBER_NAME_CHECKED(ULightComponent, LightingChannels);
+	static const FName SamplesPerPixelName = GET_MEMBER_NAME_CHECKED(ULightComponentBase, SamplesPerPixel);
+	static const FName IntensityName = GET_MEMBER_NAME_CHECKED(ULightComponentBase, Intensity);
+	static const FName IntensityUnitsName = GET_MEMBER_NAME_CHECKED(ULocalLightComponent, IntensityUnits);
+
+	// Node property names (the ones which differ from the corresponding component property names above)
+	static const FName DirectionalLightIntensityName = GET_MEMBER_NAME_CHECKED(UMovieGraphLightModifierNode, DirectionalLightIntensity);
+	static const FName SkyLightIntensityName = GET_MEMBER_NAME_CHECKED(UMovieGraphLightModifierNode, SkyLightIntensityScale);
+	static const FName PointLightIntensityName = GET_MEMBER_NAME_CHECKED(UMovieGraphLightModifierNode, PointLightIntensity);
+	static const FName RectLightIntensityName = GET_MEMBER_NAME_CHECKED(UMovieGraphLightModifierNode, RectLightIntensity);
+	static const FName SpotLightIntensityName = GET_MEMBER_NAME_CHECKED(UMovieGraphLightModifierNode, SpotLightIntensity);
+	static const FName PointLightIntensityUnitsName = GET_MEMBER_NAME_CHECKED(UMovieGraphLightModifierNode, PointLightIntensityUnits);
+	static const FName RectLightIntensityUnitsName = GET_MEMBER_NAME_CHECKED(UMovieGraphLightModifierNode, RectLightIntensityUnits);
+	static const FName SpotLightIntensityUnitsName = GET_MEMBER_NAME_CHECKED(UMovieGraphLightModifierNode, SpotLightIntensityUnits);
+
+	UMovieGraphFixedValueView* ModifierValues = Modifier->GetPropertiesForSettingValues();
+
+	// Properties are added to the modifier in one batch in order to avoid re-creating the underlying property bag constantly.
+	TArray<FMovieGraphPropertyReference> PropertiesToAdd;
+
+	// Because of the batch property creation, we have to defer setting values in the property bag until the properties have actually been created.
+	TArray<TFunction<void()>> DeferredValueSetters;
+
+	auto AddClassPredefinedLightProperty = [this, &ModifierValues, &PropertiesToAdd, &DeferredValueSetters]
+		(const bool& InOverrideProperty, const FName& InComponentPropertyName, const FName& InNodePropertyName, UClass* InLightClass, UClass* InComponentClass) -> void
+	{
+		// Nothing to do if this lighting property wasn't marked as overridden.
+		if (!InOverrideProperty)
+		{
+			return;
+		}
+
+		// The lighting properties exposed in the node always exist on the component, not the light class itself.
+		FProperty* ComponentProperty = FindFProperty<FProperty>(InComponentClass, InComponentPropertyName);
+		if (!ComponentProperty)
+		{
+			return;
+		}
+
+		FMovieGraphPropertyReference& NewProperty =
+			PropertiesToAdd.Add_GetRef(FMovieGraphPropertyReference(InLightClass, InComponentClass, InComponentPropertyName, ComponentProperty));
+
+		// Queue up setting the value in the property bag *after* the property was added to the property bag.
+		DeferredValueSetters.Add([this, &ModifierValues, NewProperty, InNodePropertyName]()
+		{
+			const FProperty* NodeProperty = FindFProperty<FProperty>(UMovieGraphLightModifierNode::StaticClass(), InNodePropertyName);
+			if (ensureMsgf(NodeProperty, TEXT("Could not find property named [%s] on the Light Modifier node."), *InNodePropertyName.ToString()))
+			{
+				FString PropertyValue;
+				NodeProperty->ExportTextItem_InContainer(PropertyValue, this, nullptr, nullptr, PPF_None);
+
+				ModifierValues->SetValueSerializedString(Modifier->GetPropertyNameForSettingCustomValue(NewProperty), PropertyValue);
+			}
+		});
+	};
+
+	/** Adds a lighting property to the modifier. This propery applies to all light types. */
+	auto AddGeneralPredefinedLightProperty = [this, &AddClassPredefinedLightProperty]
+		(const bool& InOverrideProperty, const FName& InComponentPropertyName, const FName& InNodePropertyName, UClass* InComponentClass) -> void
+	{
+		static const TArray<UClass*> LightingClasses = {
+			ADirectionalLight::StaticClass(), ASkyLight::StaticClass(), APointLight::StaticClass(), ASpotLight::StaticClass(), ARectLight::StaticClass()
+		};
+
+		// Add this property for all light classes.
+		for (UClass* LightingClass : LightingClasses)
+		{
+			AddClassPredefinedLightProperty(InOverrideProperty, InComponentPropertyName, InNodePropertyName, LightingClass, InComponentClass);
+		}
+	};
+
+	// First, add in the properties that are always defined on the node.
+	AddGeneralPredefinedLightProperty(bOverride_LightColor, LightColorName, LightColorName, ULightComponentBase::StaticClass());
+	AddGeneralPredefinedLightProperty(bOverride_bAffectsWorld, AffectsWorldName, AffectsWorldName, ULightComponentBase::StaticClass());
+	AddGeneralPredefinedLightProperty(bOverride_CastShadows, CastShadowsName, CastShadowsName, ULightComponentBase::StaticClass());
+	AddGeneralPredefinedLightProperty(bOverride_IndirectLightingIntensity, IndirectLightingIntensityName, IndirectLightingIntensityName, ULightComponentBase::StaticClass());
+	AddGeneralPredefinedLightProperty(bOverride_VolumetricScatteringIntensity, VolumetricScatteringIntensityName, VolumetricScatteringIntensityName, ULightComponentBase::StaticClass());
+	AddGeneralPredefinedLightProperty(bOverride_LightingChannels, LightingChannelsName, LightingChannelsName, ULightComponent::StaticClass());
+	AddGeneralPredefinedLightProperty(bOverride_SamplesPerPixel, SamplesPerPixelName, SamplesPerPixelName, ULightComponentBase::StaticClass());
+	AddClassPredefinedLightProperty(bOverride_DirectionalLightIntensity, IntensityName, DirectionalLightIntensityName, ADirectionalLight::StaticClass(), ULightComponentBase::StaticClass());
+	AddClassPredefinedLightProperty(bOverride_SkyLightIntensityScale, IntensityName, SkyLightIntensityName, ASkyLight::StaticClass(), ULightComponentBase::StaticClass());
+
+	if (IntensityMethod == EMovieGraphLightModifierIntensityMethod::PointRectSpot)
+	{
+		AddClassPredefinedLightProperty(bOverride_Intensity, IntensityName, IntensityName, APointLight::StaticClass(), ULightComponentBase::StaticClass());
+		AddClassPredefinedLightProperty(bOverride_Intensity, IntensityName, IntensityName, ARectLight::StaticClass(), ULightComponentBase::StaticClass());
+		AddClassPredefinedLightProperty(bOverride_Intensity, IntensityName, IntensityName, ASpotLight::StaticClass(), ULightComponentBase::StaticClass());
+		
+		AddClassPredefinedLightProperty(bOverride_IntensityUnits, IntensityUnitsName, IntensityUnitsName, APointLight::StaticClass(), ULocalLightComponent::StaticClass());
+		AddClassPredefinedLightProperty(bOverride_IntensityUnits, IntensityUnitsName, IntensityUnitsName, ARectLight::StaticClass(), ULocalLightComponent::StaticClass());
+		AddClassPredefinedLightProperty(bOverride_IntensityUnits, IntensityUnitsName, IntensityUnitsName, ASpotLight::StaticClass(), ULocalLightComponent::StaticClass());
+	}
+	else
+	{
+		AddClassPredefinedLightProperty(bOverride_PointLightIntensity, IntensityName, PointLightIntensityName, APointLight::StaticClass(), ULightComponentBase::StaticClass());
+		AddClassPredefinedLightProperty(bOverride_PointLightIntensityUnits, IntensityUnitsName, PointLightIntensityUnitsName, APointLight::StaticClass(), ULocalLightComponent::StaticClass());
+
+		AddClassPredefinedLightProperty(bOverride_RectLightIntensity, IntensityName, RectLightIntensityName, ARectLight::StaticClass(), ULightComponentBase::StaticClass());
+		AddClassPredefinedLightProperty(bOverride_RectLightIntensityUnits, IntensityUnitsName, RectLightIntensityUnitsName, ARectLight::StaticClass(), ULocalLightComponent::StaticClass());
+
+		AddClassPredefinedLightProperty(bOverride_SpotLightIntensity, IntensityName, SpotLightIntensityName, ASpotLight::StaticClass(), ULightComponentBase::StaticClass());
+		AddClassPredefinedLightProperty(bOverride_SpotLightIntensityUnits, IntensityUnitsName, SpotLightIntensityUnitsName, ASpotLight::StaticClass(), ULocalLightComponent::StaticClass());
+	}
+
+	const UMovieGraphMutableValueView* CustomLightingValues = GetDynamicPropertiesView();
+
+	// Second, add any properties that have been added in the "Custom" section.
+	for (const FName& PropertyName : CustomLightingValues->GetPropertyNames())
+	{
+		const FString PropertyNameString = PropertyName.ToString();
+
+		// Only examine bOverride_ properties. If the property isn't overridden, it won't be passed to the modifier.
+		if (!PropertyNameString.StartsWith(TEXT("bOverride_")))
+		{
+			continue;
+		}
+
+		// Check that it's overridden.
+		bool bIsOverridden = false;
+		if (!CustomLightingValues->GetValueBool(PropertyName, bIsOverridden) || !bIsOverridden)
+		{
+			continue;
+		}
+
+		FString NonOverridePropertyName = PropertyNameString;
+		NonOverridePropertyName.RemoveFromStart(TEXT("bOverride_"), ESearchCase::CaseSensitive);
+
+		// Add the Custom property to the modifier if its supplemental property info was found.
+		if (const FMovieGraphPropertyReference* CustomPropertyRef = CustomPropertyInfo.Find(FName(NonOverridePropertyName)))
+		{
+			PropertiesToAdd.Add(*CustomPropertyRef);
+
+			// Queue up setting the value in the property bag *after* the property was added to the modifier property bag.
+			FString PropertyValue = CustomLightingValues->GetValueSerializedString(FName(NonOverridePropertyName));
+			DeferredValueSetters.Add([this, &ModifierValues, CustomPropertyRef, PropertyValue]()
+			{
+				ModifierValues->SetValueSerializedString(Modifier->GetPropertyNameForSettingCustomValue(*CustomPropertyRef), PropertyValue);
+			});
+		}
+	}
+
+	// Add all the new properties in bulk, then set their values.
+	Modifier->AddProperties(PropertiesToAdd);
+	for (const TFunction<void()>& ValueSetter : DeferredValueSetters)
+	{
+		ValueSetter();
+	}
+
+	return { Modifier };
+}
+
+bool UMovieGraphLightModifierNode::SupportsCollections() const
+{
+	return true;
+}
+
+TArray<FName> UMovieGraphLightModifierNode::GetAllCollections() const
+{
+	return Collections;
+}
+
+void UMovieGraphLightModifierNode::AddCollection(const FName& InCollectionName)
+{
+#if WITH_EDITOR
+	Modify();
+#endif
+
+	if (InCollectionName == NAME_None)
+	{
+		return;
+	}
+
+	Collections.AddUnique(InCollectionName);
+}
+
+bool UMovieGraphLightModifierNode::RemoveCollection(const FName& InCollectionName)
+{
+#if WITH_EDITOR
+	Modify();
+#endif
+
+	return Collections.Remove(InCollectionName) > 0;
+}
+
+FString UMovieGraphLightModifierNode::GetNodeInstanceName() const
+{
+	return ModifierName;
+}
+
+bool UMovieGraphLightModifierNode::AddCustomLightProperty(UClass* InLightClass, const FProperty* InLightingProperty)
+{
+	const FName PropertyName = GetLightingPropertyName(InLightClass, InLightingProperty);
+
+	const UMovieGraphMutableValueView* LightingValues = GetDynamicPropertiesView();
+
+	// Skip adding if the property already exists in the property bag.
+	if (LightingValues->HasProperty(PropertyName))
+	{
+		return false;
+	}
+
+#if WITH_EDITOR
+	Modify();
+#endif
+
+	// Add supplemental property information for this property, which will importantly track the origin class and component.
+	// This has to be done via sidecar data, rather than via metadata on the property in the property bag, because runtime builds do not contain
+	// property metadata.
+	FMovieGraphPropertyReference& NewCustomProperty = CustomPropertyInfo.Emplace(
+		PropertyName,
+		FMovieGraphPropertyReference(InLightClass, InLightingProperty->GetOwnerClass(), InLightingProperty->GetFName(), const_cast<FProperty*>(InLightingProperty)));
+	NewCustomProperty.AuxPropertyName = PropertyName;	
+
+	// Update DynamicProperties with the properties referenced in CustomPropertyInfo. Although we could update DynamicProperties right here, there are
+	// parts of the graph evaluation lifecycle that depend on GetDynamicPropertyDescriptions() reflecting the property layout of DynamicProperties.
+	UpdateDynamicProperties();
+
+	return true;
+}
+
+bool UMovieGraphLightModifierNode::AddCustomLightProperty(UClass* InLightClass, const FName& InPropertyName)
+{
+	return AddCustomLightProperty(InLightClass, FindFProperty<FProperty>(InLightClass, InPropertyName));
+}
+
+bool UMovieGraphLightModifierNode::RemoveCustomLightProperty(const UClass* InLightClass, const FProperty* InLightingProperty)
+{
+	if (InLightClass && InLightingProperty)
+	{
+#if WITH_EDITOR
+		Modify();
+#endif
+
+		const FName LightingPropertyName = GetLightingPropertyName(InLightClass, InLightingProperty);
+
+		const bool bSuccess = CustomPropertyInfo.Remove(LightingPropertyName) > 0;
+
+		// Trigger a refresh of DynamicProperties
+		UpdateDynamicProperties();
+
+		return bSuccess;
+	}
+
+	return false;
+}
+
+bool UMovieGraphLightModifierNode::RemoveCustomLightProperty(const UClass* InLightClass, const FName& InPropertyName)
+{
+	return RemoveCustomLightProperty(InLightClass, FindFProperty<FProperty>(InLightClass, InPropertyName));
+}
+
+bool UMovieGraphLightModifierNode::HasCustomLightProperty(const UClass* InLightClass, const FProperty* InLightingProperty) const
+{
+	const UMovieGraphMutableValueView* LightingValues = GetDynamicPropertiesView();
+	
+	return LightingValues->HasProperty(GetLightingPropertyName(InLightClass, InLightingProperty));
+}
+
+bool UMovieGraphLightModifierNode::HasCustomLightProperty(const UClass* InLightClass, const FName& InPropertyName) const
+{
+	return HasCustomLightProperty(InLightClass, FindFProperty<FProperty>(InLightClass, InPropertyName));
+}
+
+bool UMovieGraphLightModifierNode::IsCustomLightPropertyOverridden(const UClass* InLightClass, const FName& InPropertyName) const
+{
+	const UMovieGraphMutableValueView* LightingValues = GetDynamicPropertiesView();
+	const FName MainPropertyName = GetLightingPropertyName(InLightClass, FindFProperty<FProperty>(InLightClass, InPropertyName));
+
+	bool bIsOverridden = false;
+	LightingValues->GetValueBool(GetLightingOverridePropertyName(MainPropertyName), bIsOverridden);
+
+	return bIsOverridden;
+}
+
+bool UMovieGraphLightModifierNode::UpdateCustomLightPropertyOverrideState(const UClass* InLightClass, const FName& InPropertyName, const bool bIsOverridden)
+{
+	UMovieGraphMutableValueView* LightingValues = GetDynamicPropertiesView();
+	const FName MainPropertyName = GetLightingPropertyName(InLightClass, FindFProperty<FProperty>(InLightClass, InPropertyName));
+
+	return LightingValues->SetValueBool(GetLightingOverridePropertyName(MainPropertyName), bIsOverridden);
+}
+
+FName UMovieGraphLightModifierNode::GetLightingPropertyName(const UClass* InLightClass, const FProperty* InLightingProperty)
+{
+	static const TCHAR* LightingPropertyNameFormat = TEXT("{0}_{1}");
+
+	if (InLightClass && InLightingProperty)
+	{
+		return FName(FString::Format(LightingPropertyNameFormat, { InLightClass->GetName(), InLightingProperty->GetName() }));
+	}
+
+	return FName();
+}
+
+FName UMovieGraphLightModifierNode::GetLightingOverridePropertyName(const FName& InMainPropertyName)
+{
+	static const TCHAR* LightingPropertyOverrideNameFormat = TEXT("bOverride_{0}");
+
+	return FName(FString::Format(LightingPropertyOverrideNameFormat, { InMainPropertyName.ToString() }));
+}
+
+int32 UMovieGraphLightModifierNode::GetNumCustomLightProperties() const
+{
+	const UMovieGraphMutableValueView* LightingValues = GetDynamicPropertiesView();
+
+	// The divide by two is used to account for the bOverride_* properties; those should not be counted.
+	return LightingValues->GetNumProperties() / 2;
+}
+
+#undef LOCTEXT_NAMESPACE // "MovieGraph"
\ No newline at end of file
diff --git a/Engine/Plugins/MovieScene/MovieRenderPipeline/Source/MovieRenderPipelineCore/Public/Graph/MovieGraphMultiValueContainer.cpp b/Engine/Plugins/MovieScene/MovieRenderPipeline/Source/MovieRenderPipelineCore/Public/Graph/MovieGraphMultiValueContainer.cpp
new file mode 100644
--- /dev/null
+++ b/Engine/Plugins/MovieScene/MovieRenderPipeline/Source/MovieRenderPipelineCore/Public/Graph/MovieGraphMultiValueContainer.cpp
@@ -0,0 +1,240 @@
+﻿// Copyright Epic Games, Inc. All Rights Reserved.
+
+#include "MovieGraphMultiValueContainer.h"
+
+#include "MovieGraphCommon.h"
+#include "MovieGraphValueView.h"
+
+#include UE_INLINE_GENERATED_CPP_BY_NAME(MovieGraphMultiValueContainer)
+
+UMovieGraphMultiValueContainer::UMovieGraphMultiValueContainer()
+{
+	if (!HasAnyFlags(RF_ClassDefaultObject))
+	{
+		ValuesView = NewObject<UMovieGraphMutableValueView>(this, TEXT("ValuesView"));
+		ValuesView->ViewProperties(&Values);
+	}
+}
+
+bool UMovieGraphMultiValueContainer::AddProperty(const FName& InPropertyName, const EMovieGraphValueType ValueType, UObject* InValueTypeObject, const TMap<FName, FString>& InMetadata)
+{
+	return ValuesView->AddProperty(InPropertyName, ValueType, InValueTypeObject, InMetadata);
+}
+
+bool UMovieGraphMultiValueContainer::RemoveProperty(const FName& InPropertyName)
+{
+	return ValuesView->RemoveProperty(InPropertyName);
+}
+
+bool UMovieGraphMultiValueContainer::HasProperty(const FName& InPropertyName)
+{
+	return ValuesView->HasProperty(InPropertyName);
+}
+
+int32 UMovieGraphMultiValueContainer::GetNumProperties() const
+{
+	return ValuesView->GetNumProperties();
+}
+
+TArray<FName> UMovieGraphMultiValueContainer::GetPropertyNames() const
+{
+	return ValuesView->GetPropertyNames();
+}
+
+#if WITH_EDITOR
+TMap<FName, FString> UMovieGraphMultiValueContainer::GetPropertyMetadata(const FName& InPropertyName) const
+{
+	return ValuesView->GetPropertyMetadata(InPropertyName);
+}
+#endif	// WITH_EDITOR
+
+bool UMovieGraphMultiValueContainer::GetValueBool(const FName& InPropertyName, bool& bOutValue) const
+{
+	return ValuesView->GetValueBool(InPropertyName, bOutValue);
+}
+
+bool UMovieGraphMultiValueContainer::GetValueByte(const FName& InPropertyName, uint8& OutValue) const
+{
+	return ValuesView->GetValueByte(InPropertyName, OutValue);
+}
+
+bool UMovieGraphMultiValueContainer::GetValueInt32(const FName& InPropertyName, int32& OutValue) const
+{
+	return ValuesView->GetValueInt32(InPropertyName, OutValue);
+}
+
+bool UMovieGraphMultiValueContainer::GetValueInt64(const FName& InPropertyName, int64& OutValue) const
+{
+	return ValuesView->GetValueInt64(InPropertyName, OutValue);
+}
+
+bool UMovieGraphMultiValueContainer::GetValueFloat(const FName& InPropertyName, float& OutValue) const
+{
+	return ValuesView->GetValueFloat(InPropertyName, OutValue);
+}
+
+bool UMovieGraphMultiValueContainer::GetValueDouble(const FName& InPropertyName, double& OutValue) const
+{
+	return ValuesView->GetValueDouble(InPropertyName, OutValue);
+}
+
+bool UMovieGraphMultiValueContainer::GetValueName(const FName& InPropertyName, FName& OutValue) const
+{
+	return ValuesView->GetValueName(InPropertyName, OutValue);
+}
+
+bool UMovieGraphMultiValueContainer::GetValueString(const FName& InPropertyName, FString& OutValue) const
+{
+	return ValuesView->GetValueString(InPropertyName, OutValue);
+}
+
+bool UMovieGraphMultiValueContainer::GetValueText(const FName& InPropertyName, FText& OutValue) const
+{
+	return ValuesView->GetValueText(InPropertyName, OutValue);
+}
+
+bool UMovieGraphMultiValueContainer::GetValueEnum(const FName& InPropertyName, uint8& OutValue, const UEnum* RequestedEnum) const
+{
+	return ValuesView->GetValueEnum(InPropertyName, OutValue, RequestedEnum);
+}
+
+bool UMovieGraphMultiValueContainer::GetValueStruct(const FName& InPropertyName, FStructView& OutValue, const UScriptStruct* RequestedStruct) const
+{
+	return ValuesView->GetValueStruct(InPropertyName, OutValue, RequestedStruct);
+}
+
+bool UMovieGraphMultiValueContainer::GetValueObject(const FName& InPropertyName, UObject*& OutValue, const UClass* RequestedClass) const
+{
+	return ValuesView->GetValueObject(InPropertyName, OutValue, RequestedClass);
+}
+
+bool UMovieGraphMultiValueContainer::GetValueClass(const FName& InPropertyName, UClass*& OutValue) const
+{
+	return ValuesView->GetValueClass(InPropertyName, OutValue);
+}
+
+FString UMovieGraphMultiValueContainer::GetValueSerializedString(const FName& InPropertyName)
+{
+	return ValuesView->GetValueSerializedString(InPropertyName);
+}
+
+bool UMovieGraphMultiValueContainer::SetValueBool(const FName& InPropertyName, const bool bInValue)
+{
+	return ValuesView->SetValueBool(InPropertyName, bInValue);
+}
+
+bool UMovieGraphMultiValueContainer::SetValueByte(const FName& InPropertyName, const uint8 InValue)
+{
+	return ValuesView->SetValueByte(InPropertyName, InValue);
+}
+
+bool UMovieGraphMultiValueContainer::SetValueInt32(const FName& InPropertyName, const int32 InValue)
+{
+	return ValuesView->SetValueInt32(InPropertyName, InValue);
+}
+
+bool UMovieGraphMultiValueContainer::SetValueInt64(const FName& InPropertyName, const int64 InValue)
+{
+	return ValuesView->SetValueInt64(InPropertyName, InValue);
+}
+
+bool UMovieGraphMultiValueContainer::SetValueFloat(const FName& InPropertyName, const float InValue)
+{
+	return ValuesView->SetValueFloat(InPropertyName, InValue);
+}
+
+bool UMovieGraphMultiValueContainer::SetValueDouble(const FName& InPropertyName, const double InValue)
+{
+	return ValuesView->SetValueDouble(InPropertyName, InValue);
+}
+
+bool UMovieGraphMultiValueContainer::SetValueName(const FName& InPropertyName, const FName InValue)
+{
+	return ValuesView->SetValueName(InPropertyName, InValue);
+}
+
+bool UMovieGraphMultiValueContainer::SetValueString(const FName& InPropertyName, const FString& InValue)
+{
+	return ValuesView->SetValueString(InPropertyName, InValue);
+}
+
+bool UMovieGraphMultiValueContainer::SetValueText(const FName& InPropertyName, const FText& InValue)
+{
+	return ValuesView->SetValueText(InPropertyName, InValue);
+}
+
+bool UMovieGraphMultiValueContainer::SetValueEnum(const FName& InPropertyName, const uint8 InValue, const UEnum* Enum)
+{
+	return ValuesView->SetValueEnum(InPropertyName, InValue, Enum);
+}
+
+bool UMovieGraphMultiValueContainer::SetValueStruct(const FName& InPropertyName, FConstStructView InValue)
+{
+	return ValuesView->SetValueStruct(InPropertyName, InValue);
+}
+
+bool UMovieGraphMultiValueContainer::SetValueObject(const FName& InPropertyName, UObject* InValue)
+{
+	return ValuesView->SetValueObject(InPropertyName, InValue);
+}
+
+bool UMovieGraphMultiValueContainer::SetValueClass(const FName& InPropertyName, UClass* InValue)
+{
+	return ValuesView->SetValueClass(InPropertyName, InValue);
+}
+
+bool UMovieGraphMultiValueContainer::SetValueSerializedString(const FName& InPropertyName, const FString& NewValue)
+{
+	return ValuesView->SetValueSerializedString(InPropertyName, NewValue);
+}
+
+EMovieGraphValueType UMovieGraphMultiValueContainer::GetValueType(const FName& InPropertyName) const
+{
+	return ValuesView->GetValueType(InPropertyName);
+}
+
+void UMovieGraphMultiValueContainer::SetValueType(const FName& InPropertyName, EMovieGraphValueType ValueType, UObject* InValueTypeObject)
+{
+	if (ValuesView->HasProperty(InPropertyName))
+	{
+		ValuesView->SetValueType(InPropertyName, ValueType, InValueTypeObject);
+
+#if WITH_EDITOR
+		// Send a property change event manually since the property bag doesn't seem to generate one in this scenario
+		FProperty* ValuesProperty = FindFProperty<FProperty>(StaticClass(), GET_MEMBER_NAME_CHECKED(UMovieGraphMultiValueContainer, Values));
+		FPropertyChangedEvent PropertyEvent(ValuesProperty, EPropertyChangeType::ValueSet);
+		PostEditChangeProperty(PropertyEvent);
+#endif // WITH_EDITOR
+	}
+}
+
+const UObject* UMovieGraphMultiValueContainer::GetValueTypeObject(const FName& InPropertyName) const
+{
+	return ValuesView->GetValueTypeObject(InPropertyName);
+}
+
+void UMovieGraphMultiValueContainer::SetValueTypeObject(const FName& InPropertyName, const UObject* ValueTypeObject)
+{
+	if (ValuesView->HasProperty(InPropertyName))
+    {
+    	ValuesView->SetValueTypeObject(InPropertyName, ValueTypeObject);
+    }
+}
+
+EMovieGraphContainerType UMovieGraphMultiValueContainer::GetValueContainerType(const FName& InPropertyName) const
+{
+	return ValuesView->GetValueContainerType(InPropertyName);
+}
+
+void UMovieGraphMultiValueContainer::SetValueContainerType(const FName& InPropertyName, EMovieGraphContainerType ContainerType)
+{
+	if (ValuesView->HasProperty(InPropertyName))
+	{
+		ValuesView->SetValueContainerType(InPropertyName, ContainerType);
+	}
+}
+
+TValueOrError<FPropertyBagArrayRef, EPropertyBagResult> UMovieGraphMultiValueContainer::GetArrayRef(const FName& InPropertyName)
+{
+	return ValuesView->GetArrayRef(InPropertyName);
+}
diff --git a/Engine/Plugins/MovieScene/MovieRenderPipeline/Source/MovieRenderPipelineCore/Public/Graph/MovieGraphMultiValueContainer.h b/Engine/Plugins/MovieScene/MovieRenderPipeline/Source/MovieRenderPipelineCore/Public/Graph/MovieGraphMultiValueContainer.h
new file mode 100644
--- /dev/null
+++ b/Engine/Plugins/MovieScene/MovieRenderPipeline/Source/MovieRenderPipelineCore/Public/Graph/MovieGraphMultiValueContainer.h
@@ -0,0 +1,256 @@
+﻿// Copyright Epic Games, Inc. All Rights Reserved.
+
+#pragma once
+
+#include "CoreMinimal.h"
+#include "MovieGraphCommon.h"
+#include "MovieGraphValueView.h"
+#include "StructUtils/PropertyBag.h"
+#include "Templates/ValueOrError.h"
+
+#include "MovieGraphMultiValueContainer.generated.h"
+
+#define UE_API MOVIERENDERPIPELINECORE_API
+
+/**
+ * Holds generic values, with an API for getting/setting values, as well as getting/setting the value types
+ * and value container types (eg, array).
+ *
+ * For a container that holds a single value, see UMovieGraphValueContainer.
+ */
+UCLASS(MinimalAPI)
+class UMovieGraphMultiValueContainer final : public UObject
+{
+	GENERATED_BODY()
+
+public:
+	UE_API UMovieGraphMultiValueContainer();
+
+	/**
+	 * Adds a new property to the container. If a property by the given name already exists, no add operation will occur. Returns true
+	 * on successful add, otherwise returns false.
+	 */
+	UFUNCTION(BlueprintCallable, Category="Config")
+	UE_API bool AddProperty(const FName& InPropertyName, const EMovieGraphValueType ValueType, UObject* InValueTypeObject, const TMap<FName, FString>& InMetadata);
+
+	/** Removes a property by the given name. Returns true on successful remove, otherwise returns false. */
+	UFUNCTION(BlueprintCallable, Category="Config")
+	UE_API bool RemoveProperty(const FName& InPropertyName);
+
+	/** Returns true if a property by the given name exists in the container, else false. */
+	UFUNCTION(BlueprintCallable, Category="Config")
+	UE_API bool HasProperty(const FName& InPropertyName);
+
+	/** Gets the number of properties that this container holds. */
+	UFUNCTION(BlueprintCallable, Category="Config")
+	UE_API int32 GetNumProperties() const;
+
+	/** Gets the names of properties in this container. */
+	UFUNCTION(BlueprintCallable, Category="Config")
+	UE_API TArray<FName> GetPropertyNames() const;
+
+#if WITH_EDITOR
+	/** Gets the metadata for the given property. */
+	UFUNCTION(BlueprintCallable, Category="Config")
+	UE_API TMap<FName, FString> GetPropertyMetadata(const FName& InPropertyName) const;
+#endif	// WITH_EDITOR
+
+	/** Gets the bool value of the held data. Returns true on success, else false. */
+	UFUNCTION(BlueprintCallable, Category="Config")
+	UE_API bool GetValueBool(const FName& InPropertyName, bool& bOutValue) const;
+
+	/** Gets the byte value of the held data. Returns true on success, else false. */
+	UFUNCTION(BlueprintCallable, Category="Config")
+	UE_API bool GetValueByte(const FName& InPropertyName, uint8& OutValue) const;
+
+	/** Gets the int32 value of the held data. Returns true on success, else false. */
+	UFUNCTION(BlueprintCallable, Category="Config")
+	UE_API bool GetValueInt32(const FName& InPropertyName, int32& OutValue) const;
+
+	/** Gets the int64 value of the held data. Returns true on success, else false. */
+	UFUNCTION(BlueprintCallable, Category="Config")
+	UE_API bool GetValueInt64(const FName& InPropertyName, int64& OutValue) const;
+
+	/** Gets the float value of the held data. Returns true on success, else false. */
+	UFUNCTION(BlueprintCallable, Category="Config")
+	UE_API bool GetValueFloat(const FName& InPropertyName, float& OutValue) const;
+
+	/** Gets the double value of the held data. Returns true on success, else false. */
+	UFUNCTION(BlueprintCallable, Category="Config")
+	UE_API bool GetValueDouble(const FName& InPropertyName, double& OutValue) const;
+
+	/** Gets the FName value of the held data. Returns true on success, else false. */
+	UFUNCTION(BlueprintCallable, Category="Config")
+	UE_API bool GetValueName(const FName& InPropertyName, FName& OutValue) const;
+
+	/** Gets the FString value of the held data. Returns true on success, else false. */
+	UFUNCTION(BlueprintCallable, Category="Config")
+	UE_API bool GetValueString(const FName& InPropertyName, FString& OutValue) const;
+
+	/** Gets the FText value of the held data. Returns true on success, else false. */
+	UFUNCTION(BlueprintCallable, Category="Config")
+	UE_API bool GetValueText(const FName& InPropertyName, FText& OutValue) const;
+
+	/** Gets the enum value (for a specific enum) of the held data. Returns true on success, else false. */
+	UFUNCTION(BlueprintCallable, Category="Config")
+	UE_API bool GetValueEnum(const FName& InPropertyName, uint8& OutValue, const UEnum* RequestedEnum = nullptr) const;
+
+	/** Gets the struct value (for a specific struct) of the held data. Returns true on success, else false. */
+	UE_API bool GetValueStruct(const FName& InPropertyName, FStructView& OutValue, const UScriptStruct* RequestedStruct = nullptr) const;
+
+	/** Gets the object value (for a specific class) of the held data. Returns true on success, else false. */
+	UFUNCTION(BlueprintCallable, Category="Config")
+	UE_API bool GetValueObject(const FName& InPropertyName, UObject*& OutValue, const UClass* RequestedClass = nullptr) const;
+
+	/** Gets the UClass value of the held data. Returns true on success, else false. */
+	UFUNCTION(BlueprintCallable, Category="Config")
+	UE_API bool GetValueClass(const FName& InPropertyName, UClass*& OutValue) const;
+
+	/** Gets the serialized string value of the held data. */
+	UFUNCTION(BlueprintCallable, Category="Config")
+	UE_API FString GetValueSerializedString(const FName& InPropertyName);
+
+	/** Gets the enum value of the held data. Returns true on success, else false. */
+	template <typename T>
+	bool GetValueEnum(const FName& InPropertyName, T& OutValue) const
+	{
+		return ValuesView->GetValueEnum<T>(InPropertyName, OutValue);
+	}
+
+	/** Gets the struct value of the held data. Returns true on success, else false. */
+	template <typename T>
+	bool GetValueStruct(const FName& InPropertyName, T* OutValue) const
+	{
+		return ValuesView->GetValueStruct<T>(InPropertyName, OutValue);
+	}
+
+	/** Gets the object value of the held data. Returns true on success, else false. */
+	template <typename T>
+	bool GetValueObject(const FName& InPropertyName, T*& OutValue) const
+	{
+		return ValuesView->GetValueObject<T>(InPropertyName, OutValue);
+	}
+
+	/** Sets the bool value of the held data. Returns true on success, else false. */
+	UFUNCTION(BlueprintCallable, Category="Config")
+	UE_API bool SetValueBool(const FName& InPropertyName, const bool bInValue);
+
+	/** Sets the byte value of the held data. Returns true on success, else false. */
+	UFUNCTION(BlueprintCallable, Category="Config")
+	UE_API bool SetValueByte(const FName& InPropertyName, const uint8 InValue);
+
+	/** Sets the int32 value of the held data. Returns true on success, else false. */
+	UFUNCTION(BlueprintCallable, Category="Config")
+	UE_API bool SetValueInt32(const FName& InPropertyName, const int32 InValue);
+
+	/** Sets the int64 value of the held data. Returns true on success, else false. */
+	UFUNCTION(BlueprintCallable, Category="Config")
+	UE_API bool SetValueInt64(const FName& InPropertyName, const int64 InValue);
+
+	/** Sets the float value of the held data. Returns true on success, else false. */
+	UFUNCTION(BlueprintCallable, Category="Config")
+	UE_API bool SetValueFloat(const FName& InPropertyName, const float InValue);
+
+	/** Sets the double value of the held data. Returns true on success, else false. */
+	UFUNCTION(BlueprintCallable, Category="Config")
+	UE_API bool SetValueDouble(const FName& InPropertyName, const double InValue);
+
+	/** Sets the FName value of the held data. Returns true on success, else false. */
+	UFUNCTION(BlueprintCallable, Category="Config")
+	UE_API bool SetValueName(const FName& InPropertyName, const FName InValue);
+
+	/** Sets the FString value of the held data. Returns true on success, else false. */
+	UFUNCTION(BlueprintCallable, Category="Config")
+	UE_API bool SetValueString(const FName& InPropertyName, const FString& InValue);
+
+	/** Sets the FText value of the held data. Returns true on success, else false. */
+	UFUNCTION(BlueprintCallable, Category="Config")
+	UE_API bool SetValueText(const FName& InPropertyName, const FText& InValue);
+
+	/** Sets the enum value of the held data. Returns true on success, else false. */
+	UFUNCTION(BlueprintCallable, Category="Config")
+	UE_API bool SetValueEnum(const FName& InPropertyName, const uint8 InValue, const UEnum* Enum);
+
+	/** Sets the struct value of the held data. Returns true on success, else false. */
+	UE_API bool SetValueStruct(const FName& InPropertyName, FConstStructView InValue);
+
+	/** Sets the object value of the held data. Returns true on success, else false. */
+	UFUNCTION(BlueprintCallable, Category="Config")
+	UE_API bool SetValueObject(const FName& InPropertyName, UObject* InValue);
+
+	/** Sets the class value of the held data. Returns true on success, else false. */
+	UFUNCTION(BlueprintCallable, Category="Config")
+	UE_API bool SetValueClass(const FName& InPropertyName, UClass* InValue);
+
+	/** Sets the serialized value of the held data. The string should be the serialized representation of the value. Returns true on success, else false. */
+	UFUNCTION(BlueprintCallable, Category="Config")
+	UE_API bool SetValueSerializedString(const FName& InPropertyName, const FString& NewValue);
+
+	/** Sets the enum value of the held data. Returns true on success, else false. */
+	template <typename T>
+	bool SetValueEnum(const FName& InPropertyName, const T InValue)
+	{
+		return ValuesView->SetValueEnum<T>(InPropertyName, InValue);
+	}
+
+	/** Sets the struct value of the held data. Returns true on success, else false. */
+	template <typename T>
+	bool SetValueStruct(const FName& InPropertyName, const T& InValue)
+	{
+		return ValuesView->SetValueStruct<T>(InPropertyName, InValue);
+	}
+
+	/** Sets the object value of the held data. Returns true on success, else false. */
+	template <typename T>
+	bool SetValueObject(const FName& InPropertyName, T* InValue)
+	{
+		return ValuesView->SetValueObject<T>(InPropertyName, InValue);
+	}
+
+	/** Gets the type of the stored data. */
+	UFUNCTION(BlueprintCallable, Category="Config")
+	UE_API EMovieGraphValueType GetValueType(const FName& InPropertyName) const;
+
+	/** Sets the type of the stored data. Enums, structs, and classes must specify a value type object. */
+	UFUNCTION(BlueprintCallable, Category="Config")
+	UE_API void SetValueType(const FName& InPropertyName, EMovieGraphValueType ValueType, UObject* InValueTypeObject = nullptr);
+
+	/** Gets the object that defines the enum, struct, or class. */
+	UFUNCTION(BlueprintCallable, Category="Config")
+	UE_API const UObject* GetValueTypeObject(const FName& InPropertyName) const;
+
+	/** Sets the object that defines the enum, struct, or class. */
+	UFUNCTION(BlueprintCallable, Category="Config")
+	UE_API void SetValueTypeObject(const FName& InPropertyName, const UObject* ValueTypeObject);
+
+	/** Gets the container type of the stored value. */
+	UFUNCTION(BlueprintCallable, Category="Config")
+	UE_API EMovieGraphContainerType GetValueContainerType(const FName& InPropertyName) const;
+
+	/** Sets the container type of the stored value. */
+	UFUNCTION(BlueprintCallable, Category="Config")
+	UE_API void SetValueContainerType(const FName& InPropertyName, EMovieGraphContainerType ContainerType);
+
+	/**
+	 * Gets a reference to the array backing the value, if any. GetValueContainerType() will return
+	 * EMovieGraphContainerType::Array if the value is holding an array.
+	 */
+	UE_API TValueOrError<FPropertyBagArrayRef, EPropertyBagResult> GetArrayRef(const FName& InPropertyName);
+
+private:
+	// Allow the single value container to directly modify the Values bag so we can keep the public API clean.
+	friend class UMovieGraphValueContainer;
+
+	UPROPERTY(Transient)
+	TObjectPtr<UMovieGraphMutableValueView> ValuesView = nullptr;
+
+	// Using a property bag solves a number of issues that would be very difficult to solve otherwise. 1) Presentation of
+	// a property that can have its type changed within the details pane, 2) data storage of a property which can have
+	// its type changed, 3) the ability to set the value of the property from both Python and C++, and 4) the ability to
+	// change the property at runtime.
+	/** The values held by this container. */
+	UPROPERTY(EditAnywhere, meta=(ShowOnlyInnerProperties, FixedLayout), Category = "Value")
+	FInstancedPropertyBag Values;
+};
+
+#undef UE_API
diff --git a/Engine/Plugins/MovieScene/MovieRenderPipeline/Source/MovieRenderPipelineCore/Public/Graph/MovieGraphNode.h b/Engine/Plugins/MovieScene/MovieRenderPipeline/Source/MovieRenderPipelineCore/Public/Graph/MovieGraphNode.h
--- a/Engine/Plugins/MovieScene/MovieRenderPipeline/Source/MovieRenderPipelineCore/Public/Graph/MovieGraphNode.h
+++ b/Engine/Plugins/MovieScene/MovieRenderPipeline/Source/MovieRenderPipelineCore/Public/Graph/MovieGraphNode.h
@@ -4,8 +4,8 @@
 
 #include "CoreMinimal.h"
 #include "MovieGraphPin.h"
-#include "StructUtils/InstancedStruct.h"
 #include "StructUtils/PropertyBag.h"
+#include "Graph/MovieGraphMultiValueContainer.h"
 #include "Graph/MovieGraphValueContainer.h"
 #include "Graph/MovieGraphFilenameResolveParams.h"
 #include "UObject/Interface.h"
@@ -202,6 +202,13 @@ public:
 	 */
 	MOVIERENDERPIPELINECORE_API virtual bool GetDynamicPropertyValue(const FName PropertyName, FString& OutValue);
 
+	/**
+	 * Gets the dynamic properties as a view. The view provides a scripting-friendly way of getting and modifying the values
+	 * within the dynamic properties.
+	 */
+	UFUNCTION(BlueprintCallable, Category = "Movie Graph")
+	MOVIERENDERPIPELINECORE_API UMovieGraphMutableValueView* GetDynamicPropertiesView() const;
+
 	/** Gets the override property for the specified dynamic property. If one does not exist, returns nullptr. */
 	MOVIERENDERPIPELINECORE_API const FBoolProperty* FindOverridePropertyForDynamicProperty(const FName& InPropertyName) const;
 
@@ -434,6 +441,10 @@ protected:
 	UPROPERTY(EditAnywhere, meta=(FixedLayout, ShowOnlyInnerProperties), Category = "Properties")
 	FInstancedPropertyBag DynamicProperties;
 
+	/** A property view which points to the dynamic properties on this node. */
+	UPROPERTY(Transient)
+	TObjectPtr<UMovieGraphMutableValueView> DynamicPropertiesView = nullptr;
+
 	/** Tracks which properties have been exposed on the node as inputs. */
 	UPROPERTY()
 	TArray<FMovieGraphPropertyInfo> ExposedPropertyInfo;
diff --git a/Engine/Plugins/MovieScene/MovieRenderPipeline/Source/MovieRenderPipelineCore/Public/Graph/MovieGraphValueContainer.cpp b/Engine/Plugins/MovieScene/MovieRenderPipeline/Source/MovieRenderPipelineCore/Public/Graph/MovieGraphValueContainer.cpp
--- a/Engine/Plugins/MovieScene/MovieRenderPipeline/Source/MovieRenderPipelineCore/Public/Graph/MovieGraphValueContainer.cpp
+++ b/Engine/Plugins/MovieScene/MovieRenderPipeline/Source/MovieRenderPipelineCore/Public/Graph/MovieGraphValueContainer.cpp
@@ -12,20 +12,22 @@ const FName UMovieGraphValueContainer::PropertyBagDefaultPropertyName("Value");
 UMovieGraphValueContainer::UMovieGraphValueContainer()
 {
 	PropertyName = PropertyBagDefaultPropertyName;
-	
+
+	MultiValueContainer = CreateDefaultSubobject<UMovieGraphMultiValueContainer>(TEXT("MultiValueContainer"));
+
 	if (!HasAnyFlags(RF_ClassDefaultObject))
 	{
 		// Add a default double property if one does not exist already
-		if (Value.GetNumPropertiesInBag() == 0)
+		if (MultiValueContainer->Values.GetNumPropertiesInBag() == 0)
 		{
-			Value.AddProperty(PropertyName, EPropertyBagPropertyType::Double);
+			MultiValueContainer->Values.AddProperty(PropertyName, EPropertyBagPropertyType::Double);
 		}
 	}
 }
 
 void UMovieGraphValueContainer::SetPropertyName(const FName& InName)
 {
-	if (const FPropertyBagPropertyDesc* Desc = Value.FindPropertyDescByName(PropertyName))
+	if (const FPropertyBagPropertyDesc* Desc = MultiValueContainer->Values.FindPropertyDescByName(PropertyName))
 	{
 		if (Desc->Name == InName)
 		{
@@ -40,7 +42,7 @@ void UMovieGraphValueContainer::SetPropertyName(const FName& InName)
 		NewDesc.Name = InName;
 
 		const UPropertyBag* NewPropBag = UPropertyBag::GetOrCreateFromDescs({NewDesc});
-		Value.MigrateToNewBagStruct(NewPropBag);
+		MultiValueContainer->Values.MigrateToNewBagStruct(NewPropBag);
 	}
 }
 
@@ -49,261 +51,196 @@ FName UMovieGraphValueContainer::GetPropertyName() const
 	return PropertyName;
 }
 
+void UMovieGraphValueContainer::PostLoad()
+{
+	Super::PostLoad();
+
+	if (!HasAnyFlags(RF_ClassDefaultObject))
+	{
+		// Migrate the old property bag over to the new multi-value container
+		if (Value.GetNumPropertiesInBag() > 0)
+		{
+			const TConstArrayView<FPropertyBagPropertyDesc> ExistingDescs = Value.GetPropertyBagStruct()->GetPropertyDescs();
+			MultiValueContainer->Values.AddProperties(ExistingDescs);
+			MultiValueContainer->Values.CopyMatchingValuesByName(Value);
+			Value.Reset();
+		}
+	}
+}
+
 bool UMovieGraphValueContainer::GetValueBool(bool& bOutValue) const
 {
-	TValueOrError<bool, EPropertyBagResult> Result = Value.GetValueBool(PropertyName);
-	return UE::MovieGraph::Private::GetOptionalValue<bool>(Result, bOutValue);
+	return MultiValueContainer->GetValueBool(PropertyName, bOutValue);
 }
 
 bool UMovieGraphValueContainer::GetValueByte(uint8& OutValue) const
 {
-	TValueOrError<uint8, EPropertyBagResult> Result = Value.GetValueByte(PropertyName);
-	return UE::MovieGraph::Private::GetOptionalValue<uint8>(Result, OutValue);
+	return MultiValueContainer->GetValueByte(PropertyName, OutValue);
 }
 
 bool UMovieGraphValueContainer::GetValueInt32(int32& OutValue) const
 {
-	TValueOrError<int32, EPropertyBagResult> Result = Value.GetValueInt32(PropertyName);
-	return UE::MovieGraph::Private::GetOptionalValue<int32>(Result, OutValue);
+	return MultiValueContainer->GetValueInt32(PropertyName, OutValue);
 }
 
 bool UMovieGraphValueContainer::GetValueInt64(int64& OutValue) const
 {
-	TValueOrError<int64, EPropertyBagResult> Result = Value.GetValueInt64(PropertyName);
-	return UE::MovieGraph::Private::GetOptionalValue<int64>(Result, OutValue);
+	return MultiValueContainer->GetValueInt64(PropertyName, OutValue);
 }
 
 bool UMovieGraphValueContainer::GetValueFloat(float& OutValue) const
 {
-	TValueOrError<float, EPropertyBagResult> Result = Value.GetValueFloat(PropertyName);
-	return UE::MovieGraph::Private::GetOptionalValue<float>(Result, OutValue);
+	return MultiValueContainer->GetValueFloat(PropertyName, OutValue);
 }
 
 bool UMovieGraphValueContainer::GetValueDouble(double& OutValue) const
 {
-	TValueOrError<double, EPropertyBagResult> Result = Value.GetValueDouble(PropertyName);
-	return UE::MovieGraph::Private::GetOptionalValue<double>(Result, OutValue);
+	return MultiValueContainer->GetValueDouble(PropertyName, OutValue);
 }
 
 bool UMovieGraphValueContainer::GetValueName(FName& OutValue) const
 {
-	TValueOrError<FName, EPropertyBagResult> Result = Value.GetValueName(PropertyName);
-	return UE::MovieGraph::Private::GetOptionalValue<FName>(Result, OutValue);
+	return MultiValueContainer->GetValueName(PropertyName, OutValue);
 }
 
 bool UMovieGraphValueContainer::GetValueString(FString& OutValue) const
 {
-	TValueOrError<FString, EPropertyBagResult> Result = Value.GetValueString(PropertyName);
-	return UE::MovieGraph::Private::GetOptionalValue<FString>(Result, OutValue);
+	return MultiValueContainer->GetValueString(PropertyName, OutValue);
 }
 
 bool UMovieGraphValueContainer::GetValueText(FText& OutValue) const
 {
-	TValueOrError<FText, EPropertyBagResult> Result = Value.GetValueText(PropertyName);
-	return UE::MovieGraph::Private::GetOptionalValue<FText>(Result, OutValue);
+	return MultiValueContainer->GetValueText(PropertyName, OutValue);
 }
 
 bool UMovieGraphValueContainer::GetValueEnum(uint8& OutValue, const UEnum* RequestedEnum) const
 {
-	TValueOrError<uint8, EPropertyBagResult> Result = Value.GetValueEnum(PropertyName, RequestedEnum);
-	return UE::MovieGraph::Private::GetOptionalValue<uint8>(Result, OutValue);
+	return MultiValueContainer->GetValueEnum(PropertyName, OutValue, RequestedEnum);
 }
 
 bool UMovieGraphValueContainer::GetValueStruct(FStructView& OutValue, const UScriptStruct* RequestedStruct) const
 {
-	TValueOrError<FStructView, EPropertyBagResult> Result = Value.GetValueStruct(PropertyName, RequestedStruct);
-	return UE::MovieGraph::Private::GetOptionalValue<FStructView>(Result, OutValue);
+	return MultiValueContainer->GetValueStruct(PropertyName, OutValue, RequestedStruct);
 }
 
-bool UMovieGraphValueContainer::GetValueObject(UObject* OutValue, const UClass* RequestedClass) const
+bool UMovieGraphValueContainer::GetValueObject(UObject*& OutValue, const UClass* RequestedClass) const
 {
-	TValueOrError<UObject*, EPropertyBagResult> Result = Value.GetValueObject(PropertyName, RequestedClass);
-	return UE::MovieGraph::Private::GetOptionalValue<UObject*>(Result, OutValue);
+	return MultiValueContainer->GetValueObject(PropertyName, OutValue, RequestedClass);
 }
 
 bool UMovieGraphValueContainer::GetValueClass(UClass*& OutValue) const
 {
-	TValueOrError<UClass*, EPropertyBagResult> Result = Value.GetValueClass(PropertyName);
-	return UE::MovieGraph::Private::GetOptionalValue<UClass*>(Result, OutValue);
+	return MultiValueContainer->GetValueClass(PropertyName, OutValue);
 }
 
 FString UMovieGraphValueContainer::GetValueSerializedString()
 {
-	TValueOrError<FString, EPropertyBagResult> Result = Value.GetValueSerializedString(PropertyName);
-	FString ResultString;
-	UE::MovieGraph::Private::GetOptionalValue<FString>(Result, ResultString);
-	return ResultString;
+	return MultiValueContainer->GetValueSerializedString(PropertyName);
 }
 
 bool UMovieGraphValueContainer::SetValueBool(const bool bInValue)
 {
-	Modify();
-	return Value.SetValueBool(PropertyName, bInValue) == EPropertyBagResult::Success;
+	return MultiValueContainer->SetValueBool(PropertyName, bInValue);
 }
 
 bool UMovieGraphValueContainer::SetValueByte(const uint8 InValue)
 {
-	Modify();
-	return Value.SetValueByte(PropertyName, InValue) == EPropertyBagResult::Success;
+	return MultiValueContainer->SetValueByte(PropertyName, InValue);
 }
 
 bool UMovieGraphValueContainer::SetValueInt32(const int32 InValue)
 {
-	Modify();
-	return Value.SetValueInt32(PropertyName, InValue) == EPropertyBagResult::Success;
+	return MultiValueContainer->SetValueInt32(PropertyName, InValue);
 }
 
 bool UMovieGraphValueContainer::SetValueInt64(const int64 InValue)
 {
-	Modify();
-	return Value.SetValueInt64(PropertyName, InValue) == EPropertyBagResult::Success;
+	return MultiValueContainer->SetValueInt64(PropertyName, InValue);
 }
 
 bool UMovieGraphValueContainer::SetValueFloat(const float InValue)
 {
-	Modify();
-	return Value.SetValueFloat(PropertyName, InValue) == EPropertyBagResult::Success;
+	return MultiValueContainer->SetValueFloat(PropertyName, InValue);
 }
 
 bool UMovieGraphValueContainer::SetValueDouble(const double InValue)
 {
-	Modify();
-	return Value.SetValueDouble(PropertyName, InValue) == EPropertyBagResult::Success;
+	return MultiValueContainer->SetValueDouble(PropertyName, InValue);
 }
 
 bool UMovieGraphValueContainer::SetValueName(const FName InValue)
 {
-	Modify();
-	return Value.SetValueName(PropertyName, InValue) == EPropertyBagResult::Success;
+	return MultiValueContainer->SetValueName(PropertyName, InValue);
 }
 
 bool UMovieGraphValueContainer::SetValueString(const FString& InValue)
 {
-	Modify();
-	return Value.SetValueString(PropertyName, InValue) == EPropertyBagResult::Success;
+	return MultiValueContainer->SetValueString(PropertyName, InValue);
 }
 
 bool UMovieGraphValueContainer::SetValueText(const FText& InValue)
 {
-	Modify();
-	return Value.SetValueText(PropertyName, InValue) == EPropertyBagResult::Success;
+	return MultiValueContainer->SetValueText(PropertyName, InValue);
 }
 
 bool UMovieGraphValueContainer::SetValueEnum(const uint8 InValue, const UEnum* Enum)
 {
-	Modify();
-	return Value.SetValueEnum(PropertyName, InValue, Enum) == EPropertyBagResult::Success;
+	return MultiValueContainer->SetValueEnum(PropertyName, InValue, Enum);
 }
 
 bool UMovieGraphValueContainer::SetValueStruct(FConstStructView InValue)
 {
-	Modify();
-	return Value.SetValueStruct(PropertyName, InValue) == EPropertyBagResult::Success;
+	return MultiValueContainer->SetValueStruct(PropertyName, InValue);
 }
 
 bool UMovieGraphValueContainer::SetValueObject(UObject* InValue)
 {
-	Modify();
-	return Value.SetValueObject(PropertyName, InValue) == EPropertyBagResult::Success;
+	return MultiValueContainer->SetValueObject(PropertyName, InValue);
 }
 
 bool UMovieGraphValueContainer::SetValueClass(UClass* InValue)
 {
-	Modify();
-	return Value.SetValueClass(PropertyName, InValue) == EPropertyBagResult::Success;
+	return MultiValueContainer->SetValueClass(PropertyName, InValue);
 }
 
 bool UMovieGraphValueContainer::SetValueSerializedString(const FString& NewValue)
 {
-	Modify();
-	return Value.SetValueSerializedString(PropertyName, NewValue) == EPropertyBagResult::Success;
+	return MultiValueContainer->SetValueSerializedString(PropertyName, NewValue);
 }
 
 EMovieGraphValueType UMovieGraphValueContainer::GetValueType() const
 {
-	if (const FPropertyBagPropertyDesc* Desc = Value.FindPropertyDescByName(PropertyName))
-	{
-		return static_cast<EMovieGraphValueType>(Desc->ValueType);
-	}
-
-	return EMovieGraphValueType::None;
+	return MultiValueContainer->GetValueType(PropertyName);
 }
 
 void UMovieGraphValueContainer::SetValueType(EMovieGraphValueType ValueType, UObject* InValueTypeObject)
 {
-	if (const FPropertyBagPropertyDesc* Desc = Value.FindPropertyDescByName(PropertyName))
-	{
-		Modify();
-
-		FPropertyBagPropertyDesc NewDesc;
-		NewDesc.Name = Desc->Name;
-		NewDesc.ValueType = static_cast<EPropertyBagPropertyType>(ValueType);
-		NewDesc.ContainerTypes = Desc->ContainerTypes;
-		NewDesc.ValueTypeObject = InValueTypeObject;
-
-		Value.Reset();
-		Value.AddProperties({NewDesc});
-
-#if WITH_EDITOR
-		// Send a property change event manually since the property bag doesn't seem to generate one in this scenario
-		FProperty* ValueProperty = FindFProperty<FProperty>(StaticClass(), GET_MEMBER_NAME_CHECKED(UMovieGraphValueContainer, Value));
-		FPropertyChangedEvent PropertyEvent(ValueProperty, EPropertyChangeType::ValueSet);
-		PostEditChangeProperty(PropertyEvent);
-#endif // WITH_EDITOR
-	}
+	return MultiValueContainer->SetValueType(PropertyName, ValueType, InValueTypeObject);
 }
 
 const UObject* UMovieGraphValueContainer::GetValueTypeObject() const
 {
-	if (const FPropertyBagPropertyDesc* Desc = Value.FindPropertyDescByName(PropertyName))
-	{
-		return Desc->ValueTypeObject;
-	}
-	
-	return nullptr;
+	return MultiValueContainer->GetValueTypeObject(PropertyName);
 }
 
 void UMovieGraphValueContainer::SetValueTypeObject(const UObject* ValueTypeObject)
 {
-	if (const FPropertyBagPropertyDesc* Desc = Value.FindPropertyDescByName(PropertyName))
-	{
-		Modify();
-
-		FPropertyBagPropertyDesc NewDesc(*Desc);
-		NewDesc.ValueTypeObject = ValueTypeObject;
-
-		const UPropertyBag* NewPropBag = UPropertyBag::GetOrCreateFromDescs({NewDesc});
-		Value.MigrateToNewBagStruct(NewPropBag);
-	}
+	return MultiValueContainer->SetValueTypeObject(PropertyName, ValueTypeObject);
 }
 
 EMovieGraphContainerType UMovieGraphValueContainer::GetValueContainerType() const
 {
-	if (const FPropertyBagPropertyDesc* Desc = Value.FindPropertyDescByName(PropertyName))
-	{
-		return static_cast<EMovieGraphContainerType>(Desc->ContainerTypes.GetFirstContainerType());
-	}
-
-	return EMovieGraphContainerType::None;
+	return MultiValueContainer->GetValueContainerType(PropertyName);
 }
 
 void UMovieGraphValueContainer::SetValueContainerType(EMovieGraphContainerType ContainerType)
 {
-	if (const FPropertyBagPropertyDesc* Desc = Value.FindPropertyDescByName(PropertyName))
-	{
-		Modify();
-
-		FPropertyBagPropertyDesc NewDesc(*Desc);
-		NewDesc.ContainerTypes = { static_cast<EPropertyBagContainerType>(ContainerType) };
-
-		const UPropertyBag* NewPropBag = UPropertyBag::GetOrCreateFromDescs({NewDesc});
-		Value.MigrateToNewBagStruct(NewPropBag);
-	}
+	return MultiValueContainer->SetValueContainerType(PropertyName, ContainerType);
 }
 
 TValueOrError<FPropertyBagArrayRef, EPropertyBagResult> UMovieGraphValueContainer::GetArrayRef()
 {
-	return Value.GetMutableArrayRef(PropertyName);
+	return MultiValueContainer->GetArrayRef(PropertyName);
 }
 
 void UMovieGraphValueContainer::SetFromDesc(const FPropertyBagPropertyDesc* InDesc, const FString& InString)
@@ -319,9 +256,9 @@ void UMovieGraphValueContainer::SetFromDesc(const FPropertyBagPropertyDesc* InDe
 		NewDesc.ValueTypeObject = InDesc->ValueTypeObject;
 
 		PropertyName = InDesc->Name;
-		
-		Value.Reset();
-		Value.AddProperties({NewDesc});
+
+		MultiValueContainer->Values.Reset();
+		MultiValueContainer->Values.AddProperties({NewDesc});
 
 		SetValueSerializedString(InString);
 	}
diff --git a/Engine/Plugins/MovieScene/MovieRenderPipeline/Source/MovieRenderPipelineCore/Public/Graph/MovieGraphValueContainer.h b/Engine/Plugins/MovieScene/MovieRenderPipeline/Source/MovieRenderPipelineCore/Public/Graph/MovieGraphValueContainer.h
--- a/Engine/Plugins/MovieScene/MovieRenderPipeline/Source/MovieRenderPipelineCore/Public/Graph/MovieGraphValueContainer.h
+++ b/Engine/Plugins/MovieScene/MovieRenderPipeline/Source/MovieRenderPipelineCore/Public/Graph/MovieGraphValueContainer.h
@@ -4,6 +4,7 @@
 
 #include "CoreMinimal.h"
 #include "MovieGraphCommon.h"
+#include "MovieGraphMultiValueContainer.h"
 #include "StructUtils/PropertyBag.h"
 #include "Templates/ValueOrError.h"
 
@@ -11,26 +12,11 @@
 
 #define UE_API MOVIERENDERPIPELINECORE_API
 
-namespace UE::MovieGraph::Private
-{
-	template<typename ReturnType>
-	bool GetOptionalValue(TValueOrError<ReturnType, EPropertyBagResult>& PropertyBagValue, ReturnType& OutValue)
-	{
-		// Convert the property bag-provided TValueOrError, which contains an EPropertyBagResult, to an output value and
-		// a bool (signifying if there was an error or not). EPropertyBagResult shouldn't be exposed on the MRQ API.
-		if (PropertyBagValue.HasValue())
-		{
-			OutValue = PropertyBagValue.StealValue();
-			return true;
-		}
-
-		return false;
-	}
-}
-
 /**
  * Holds a generic value, with an API for getting/setting the value, as well as getting/setting its type
- * and container (eg, array).
+ * and container type (eg, array).
+ *
+ * For a container that holds multiple values, see UMovieGraphMultiValueContainer.
  */
 UCLASS(MinimalAPI)
 class UMovieGraphValueContainer : public UObject
@@ -85,7 +71,7 @@ public:
 
 	/** Gets the object value (for a specific class) of the held data. Returns true on success, else false. */
 	UFUNCTION(BlueprintCallable, Category="Config")
-	UE_API bool GetValueObject(UObject* OutValue, const UClass* RequestedClass = nullptr) const;
+	UE_API bool GetValueObject(UObject*& OutValue, const UClass* RequestedClass = nullptr) const;
 
 	/** Gets the UClass value of the held data. Returns true on success, else false. */
 	UFUNCTION(BlueprintCallable, Category="Config")
@@ -99,24 +85,21 @@ public:
 	template <typename T>
 	bool GetValueEnum(T& OutValue) const
 	{
-		TValueOrError<T, EPropertyBagResult> Result = Value.GetValueEnum<T>(PropertyBagDefaultPropertyName);
-		return UE::MovieGraph::Private::GetOptionalValue<T>(Result, OutValue);
+		return GetValueEnum(PropertyName, OutValue);
 	}
 
 	/** Gets the struct value of the held data. Returns true on success, else false. */
 	template <typename T>
 	bool GetValueStruct(T* OutValue) const
 	{
-		TValueOrError<T*, EPropertyBagResult> Result = Value.GetValueStruct<T*>(PropertyBagDefaultPropertyName);
-		return UE::MovieGraph::Private::GetOptionalValue<T*>(Result, OutValue);
+		return MultiValueContainer->GetValueStruct(PropertyName, OutValue);
 	}
 
 	/** Gets the object value of the held data. Returns true on success, else false. */
 	template <typename T>
-	bool GetValueObject(T* OutValue) const
+	bool GetValueObject(T*& OutValue) const
 	{
-		TValueOrError<T*, EPropertyBagResult> Result = Value.GetValueObject<T*>(PropertyBagDefaultPropertyName);
-		return UE::MovieGraph::Private::GetOptionalValue<T*>(Result, OutValue);
+		return MultiValueContainer->GetValueObject(PropertyName, OutValue);
 	}
 
 	/** Sets the bool value of the held data. Returns true on success, else false. */
@@ -178,21 +161,21 @@ public:
 	template <typename T>
 	bool SetValueEnum(const T InValue)
 	{
-		return Value.SetValueEnum<T>(PropertyBagDefaultPropertyName, InValue) == EPropertyBagResult::Success;
+		return MultiValueContainer->SetValueEnum<T>(PropertyName, InValue);
 	}
 
 	/** Sets the struct value of the held data. Returns true on success, else false. */
 	template <typename T>
 	bool SetValueStruct(const T& InValue)
 	{
-		return Value.SetValueStruct<T>(PropertyBagDefaultPropertyName, InValue) == EPropertyBagResult::Success;
+		return MultiValueContainer->SetValueStruct<T>(PropertyName, InValue);
 	}
 
 	/** Sets the object value of the held data. Returns true on success, else false. */
 	template <typename T>
 	bool SetValueObject(T* InValue)
 	{
-		return Value.SetValueObject<T>(PropertyBagDefaultPropertyName, InValue) == EPropertyBagResult::Success;
+		return MultiValueContainer->SetValueObject<T>(PropertyName, InValue);
 	}
 
 	/** Gets the type of the stored data. */
@@ -235,6 +218,8 @@ public:
 	UE_API FName GetPropertyName() const;
 
 private:
+	virtual void PostLoad() override;
+	
 	/**
 	 * Sets the configuration of this container from a property desc. This is less safe than the strongly-typed methods,
 	 * and is reserved only for UMovieJobVariableAssignmentContainer to call. The serialized representation of the value
@@ -251,14 +236,18 @@ private:
 	UPROPERTY()
 	FName PropertyName;
 	
-	// Note: The property bag only stores one property, since the object only needs to store one value. This is an odd
-	// use of a property bag, but the property bag solves a number of issues that would be very difficult to solve
-	// otherwise. 1) Presentation of a property that can have its type changed within the details pane, 2) data storage
-	// of a property which can have its type changed, 3) the ability to set the value of the property from both Python
-	// and C++, and 4) the ability to change the property at runtime.
-	/** The value held by this object. */
-	UPROPERTY(EditAnywhere, meta=(ShowOnlyInnerProperties, FixedLayout), Category = "Value")
+	/** UNUSED (kept for backwards compatibility purposes in PostLoad()). The value held by this object. */
+	UPROPERTY()
 	FInstancedPropertyBag Value;
+
+	/**
+	 * The container that holds the single value.
+	 *
+	 * It's strange for a single-value class to internally hold a multi-value container. However, this was done in order to centralize the core
+	 * logic of how MRG interacts with property bags in primarily one class (the multi-value container).
+	 */
+	UPROPERTY(EditAnywhere, Instanced, Category = "Value")
+	TObjectPtr<UMovieGraphMultiValueContainer> MultiValueContainer;
 };
 
 #undef UE_API
diff --git a/Engine/Plugins/MovieScene/MovieRenderPipeline/Source/MovieRenderPipelineCore/Public/Graph/MovieGraphValueUtils.h b/Engine/Plugins/MovieScene/MovieRenderPipeline/Source/MovieRenderPipelineCore/Public/Graph/MovieGraphValueUtils.h
new file mode 100644
--- /dev/null
+++ b/Engine/Plugins/MovieScene/MovieRenderPipeline/Source/MovieRenderPipelineCore/Public/Graph/MovieGraphValueUtils.h
@@ -0,0 +1,12 @@
+﻿// Copyright Epic Games, Inc. All Rights Reserved.
+
+#pragma once
+
+#include "CoreMinimal.h"
+#include "MovieGraphCommon.h"
+
+namespace UE::MovieGraph::ValueUtils
+{
+	/** Gets the value type and value type object for a given FProperty. Returns true if the types could be determined, else false. */
+	bool GetTypesForFProperty(const FProperty* InProperty, EMovieGraphValueType& OutValueType, TObjectPtr<UObject>& OutValueTypeObject);
+}
diff --git a/Engine/Plugins/MovieScene/MovieRenderPipeline/Source/MovieRenderPipelineCore/Public/Graph/MovieGraphValueView.cpp b/Engine/Plugins/MovieScene/MovieRenderPipeline/Source/MovieRenderPipelineCore/Public/Graph/MovieGraphValueView.cpp
new file mode 100644
--- /dev/null
+++ b/Engine/Plugins/MovieScene/MovieRenderPipeline/Source/MovieRenderPipelineCore/Public/Graph/MovieGraphValueView.cpp
@@ -0,0 +1,430 @@
+﻿// Copyright Epic Games, Inc. All Rights Reserved.
+
+#include "MovieGraphValueView.h"
+
+#include "MovieGraphCommon.h"
+#include "MovieRenderPipelineCoreModule.h"
+
+#include UE_INLINE_GENERATED_CPP_BY_NAME(MovieGraphValueView)
+
+FMovieGraphValueViewProperty::FMovieGraphValueViewProperty(
+	const FName& InPropertyName, const EMovieGraphValueType ValueType, const TObjectPtr<UObject>& InValueTypeObject, const TMap<FName, FString>& InMetadata)
+	: PropertyName(InPropertyName)
+	, ValueType(ValueType)
+	, ValueTypeObject(InValueTypeObject)
+	, Metadata(InMetadata)
+{
+	
+}
+
+UMovieGraphValueView::UMovieGraphValueView(FInstancedPropertyBag* InPropertyBag)
+{
+	PropertyBag = InPropertyBag;
+}
+
+void UMovieGraphValueView::ViewProperties(FInstancedPropertyBag* InPropertyBag)
+{
+	PropertyBag = InPropertyBag;
+}
+
+bool UMovieGraphValueView::IsValid() const
+{
+	return PropertyBag && PropertyBag->IsValid();
+}
+
+bool UMovieGraphValueView::AddProperty(const FName& InPropertyName, const EMovieGraphValueType ValueType, UObject* InValueTypeObject, const TMap<FName, FString>& InMetadata)
+{
+	if (!bAllowPropertyAddRemove)
+	{
+		UE_LOG(LogMovieRenderPipeline, Warning, TEXT("Adding properties via this value view is not allowed; it was marked as non-mutable."));
+		return false;
+	}
+
+	Modify();
+	
+	FMovieGraphValueViewProperty NewProperty;
+	NewProperty.PropertyName = InPropertyName;
+	NewProperty.ValueType = ValueType;
+	NewProperty.ValueTypeObject = InValueTypeObject;
+	NewProperty.Metadata = InMetadata;
+
+	TArray<FMovieGraphValueViewProperty> NewProperties;
+	NewProperties.Emplace(MoveTemp(NewProperty));
+
+	return AddProperties(NewProperties);
+}
+
+bool UMovieGraphValueView::AddProperties(const TArray<FMovieGraphValueViewProperty>& InProperties)
+{
+	if (!bAllowPropertyAddRemove)
+	{
+		UE_LOG(LogMovieRenderPipeline, Warning, TEXT("Adding properties via this value view is not allowed; it was marked as non-mutable."));
+		return false;
+	}
+
+	Modify();
+
+	TArray<FPropertyBagPropertyDesc> NewPropertyDescs;
+
+	for (const FMovieGraphValueViewProperty& NewProperty : InProperties)
+	{
+		FPropertyBagPropertyDesc NewDesc(
+			NewProperty.PropertyName, static_cast<EPropertyBagPropertyType>(NewProperty.ValueType), NewProperty.ValueTypeObject);
+
+		NewDesc.ID = FGuid::NewGuid();
+
+#if WITH_EDITOR
+		for (const TPair<FName, FString>& MetadataPair : NewProperty.Metadata)
+		{
+			NewDesc.MetaData.Add(FPropertyBagPropertyDescMetaData(MetadataPair.Key, MetadataPair.Value));
+		}
+#endif
+
+		NewPropertyDescs.Add(NewDesc);
+	}
+
+	return PropertyBag->AddProperties(NewPropertyDescs) == EPropertyBagAlterationResult::Success;
+}
+
+bool UMovieGraphValueView::RemoveProperty(const FName& InPropertyName)
+{
+	if (!bAllowPropertyAddRemove)
+	{
+		UE_LOG(LogMovieRenderPipeline, Warning, TEXT("Removing properties via this value view is not allowed; it was marked as non-mutable."));
+		return false;
+	}
+
+	Modify();
+
+	const EPropertyBagAlterationResult Result = PropertyBag->RemovePropertyByName(InPropertyName);
+	return Result == EPropertyBagAlterationResult::Success;
+}
+
+void UMovieGraphValueView::RemoveAllProperties()
+{
+	if (!bAllowPropertyAddRemove)
+	{
+		UE_LOG(LogMovieRenderPipeline, Warning, TEXT("Removing properties via this value view is not allowed; it was marked as non-mutable."));
+		return;
+	}
+
+	Modify();
+
+	PropertyBag->Reset();
+}
+
+bool UMovieGraphValueView::HasProperty(const FName& InPropertyName) const
+{
+	return PropertyBag->FindPropertyDescByName(InPropertyName) != nullptr;
+}
+
+int32 UMovieGraphValueView::GetNumProperties() const
+{
+	return PropertyBag->GetNumPropertiesInBag();
+}
+
+TArray<FName> UMovieGraphValueView::GetPropertyNames() const
+{
+	TArray<FName> PropertyNames;
+
+	if (const UPropertyBag* Bag = PropertyBag->GetPropertyBagStruct())
+	{
+		TConstArrayView<FPropertyBagPropertyDesc> PropertyDescs = Bag->GetPropertyDescs();
+		PropertyNames.Reserve(PropertyDescs.Num());
+
+		for (const FPropertyBagPropertyDesc& Desc : PropertyDescs)
+		{
+			PropertyNames.Add(Desc.Name);
+		}
+	}
+
+	return PropertyNames;	
+}
+
+#if WITH_EDITOR
+TMap<FName, FString> UMovieGraphValueView::GetPropertyMetadata(const FName& InPropertyName) const
+{
+	TMap<FName, FString> MetadataMap;
+
+	if (const FPropertyBagPropertyDesc* Desc = PropertyBag->FindPropertyDescByName(InPropertyName))
+	{
+		for (const FPropertyBagPropertyDescMetaData& MetadataEntry : Desc->MetaData)
+		{
+			MetadataMap.Add(MetadataEntry.Key, MetadataEntry.Value);
+		}
+	}
+
+	return MetadataMap;
+}
+#endif	// WITH_EDITOR
+
+bool UMovieGraphValueView::GetValueBool(const FName& InPropertyName, bool& bOutValue) const
+{
+	TValueOrError<bool, EPropertyBagResult> Result = PropertyBag->GetValueBool(InPropertyName);
+	return UE::MovieGraph::Private::GetOptionalValue<bool>(Result, bOutValue);
+}
+
+bool UMovieGraphValueView::GetValueByte(const FName& InPropertyName, uint8& OutValue) const
+{
+	TValueOrError<uint8, EPropertyBagResult> Result = PropertyBag->GetValueByte(InPropertyName);
+	return UE::MovieGraph::Private::GetOptionalValue<uint8>(Result, OutValue);
+}
+
+bool UMovieGraphValueView::GetValueInt32(const FName& InPropertyName, int32& OutValue) const
+{
+	TValueOrError<int32, EPropertyBagResult> Result = PropertyBag->GetValueInt32(InPropertyName);
+	return UE::MovieGraph::Private::GetOptionalValue<int32>(Result, OutValue);
+}
+
+bool UMovieGraphValueView::GetValueInt64(const FName& InPropertyName, int64& OutValue) const
+{
+	TValueOrError<int64, EPropertyBagResult> Result = PropertyBag->GetValueInt64(InPropertyName);
+	return UE::MovieGraph::Private::GetOptionalValue<int64>(Result, OutValue);
+}
+
+bool UMovieGraphValueView::GetValueFloat(const FName& InPropertyName, float& OutValue) const
+{
+	TValueOrError<float, EPropertyBagResult> Result = PropertyBag->GetValueFloat(InPropertyName);
+	return UE::MovieGraph::Private::GetOptionalValue<float>(Result, OutValue);
+}
+
+bool UMovieGraphValueView::GetValueDouble(const FName& InPropertyName, double& OutValue) const
+{
+	TValueOrError<double, EPropertyBagResult> Result = PropertyBag->GetValueDouble(InPropertyName);
+	return UE::MovieGraph::Private::GetOptionalValue<double>(Result, OutValue);
+}
+
+bool UMovieGraphValueView::GetValueName(const FName& InPropertyName, FName& OutValue) const
+{
+	TValueOrError<FName, EPropertyBagResult> Result = PropertyBag->GetValueName(InPropertyName);
+	return UE::MovieGraph::Private::GetOptionalValue<FName>(Result, OutValue);
+}
+
+bool UMovieGraphValueView::GetValueString(const FName& InPropertyName, FString& OutValue) const
+{
+	TValueOrError<FString, EPropertyBagResult> Result = PropertyBag->GetValueString(InPropertyName);
+	return UE::MovieGraph::Private::GetOptionalValue<FString>(Result, OutValue);
+}
+
+bool UMovieGraphValueView::GetValueText(const FName& InPropertyName, FText& OutValue) const
+{
+	TValueOrError<FText, EPropertyBagResult> Result = PropertyBag->GetValueText(InPropertyName);
+	return UE::MovieGraph::Private::GetOptionalValue<FText>(Result, OutValue);
+}
+
+bool UMovieGraphValueView::GetValueEnum(const FName& InPropertyName, uint8& OutValue, const UEnum* RequestedEnum) const
+{
+	TValueOrError<uint8, EPropertyBagResult> Result = PropertyBag->GetValueEnum(InPropertyName, RequestedEnum);
+	return UE::MovieGraph::Private::GetOptionalValue<uint8>(Result, OutValue);
+}
+
+bool UMovieGraphValueView::GetValueStruct(const FName& InPropertyName, FStructView& OutValue, const UScriptStruct* RequestedStruct) const
+{
+	TValueOrError<FStructView, EPropertyBagResult> Result = PropertyBag->GetValueStruct(InPropertyName, RequestedStruct);
+	return UE::MovieGraph::Private::GetOptionalValue<FStructView>(Result, OutValue);
+}
+
+bool UMovieGraphValueView::GetValueObject(const FName& InPropertyName, UObject*& OutValue, const UClass* RequestedClass) const
+{
+	TValueOrError<UObject*, EPropertyBagResult> Result = PropertyBag->GetValueObject(InPropertyName, RequestedClass);
+	return UE::MovieGraph::Private::GetOptionalValue<UObject*>(Result, OutValue);
+}
+
+bool UMovieGraphValueView::GetValueClass(const FName& InPropertyName, UClass*& OutValue) const
+{
+	TValueOrError<UClass*, EPropertyBagResult> Result = PropertyBag->GetValueClass(InPropertyName);
+	return UE::MovieGraph::Private::GetOptionalValue<UClass*>(Result, OutValue);
+}
+
+FString UMovieGraphValueView::GetValueSerializedString(const FName& InPropertyName) const
+{
+	TValueOrError<FString, EPropertyBagResult> Result = PropertyBag->GetValueSerializedString(InPropertyName);
+	FString ResultString;
+	UE::MovieGraph::Private::GetOptionalValue<FString>(Result, ResultString);
+	return ResultString;
+}
+
+bool UMovieGraphValueView::SetValueBool(const FName& InPropertyName, const bool bInValue)
+{
+	Modify();
+	return PropertyBag->SetValueBool(InPropertyName, bInValue) == EPropertyBagResult::Success;
+}
+
+bool UMovieGraphValueView::SetValueByte(const FName& InPropertyName, const uint8 InValue)
+{
+	Modify();
+	return PropertyBag->SetValueByte(InPropertyName, InValue) == EPropertyBagResult::Success;
+}
+
+bool UMovieGraphValueView::SetValueInt32(const FName& InPropertyName, const int32 InValue)
+{
+	Modify();
+	return PropertyBag->SetValueInt32(InPropertyName, InValue) == EPropertyBagResult::Success;
+}
+
+bool UMovieGraphValueView::SetValueInt64(const FName& InPropertyName, const int64 InValue)
+{
+	Modify();
+	return PropertyBag->SetValueInt64(InPropertyName, InValue) == EPropertyBagResult::Success;
+}
+
+bool UMovieGraphValueView::SetValueFloat(const FName& InPropertyName, const float InValue)
+{
+	Modify();
+	return PropertyBag->SetValueFloat(InPropertyName, InValue) == EPropertyBagResult::Success;
+}
+
+bool UMovieGraphValueView::SetValueDouble(const FName& InPropertyName, const double InValue)
+{
+	Modify();
+	return PropertyBag->SetValueDouble(InPropertyName, InValue) == EPropertyBagResult::Success;
+}
+
+bool UMovieGraphValueView::SetValueName(const FName& InPropertyName, const FName InValue)
+{
+	Modify();
+	return PropertyBag->SetValueName(InPropertyName, InValue) == EPropertyBagResult::Success;
+}
+
+bool UMovieGraphValueView::SetValueString(const FName& InPropertyName, const FString& InValue)
+{
+	Modify();
+	return PropertyBag->SetValueString(InPropertyName, InValue) == EPropertyBagResult::Success;
+}
+
+bool UMovieGraphValueView::SetValueText(const FName& InPropertyName, const FText& InValue)
+{
+	Modify();
+	return PropertyBag->SetValueText(InPropertyName, InValue) == EPropertyBagResult::Success;
+}
+
+bool UMovieGraphValueView::SetValueEnum(const FName& InPropertyName, const uint8 InValue, const UEnum* Enum)
+{
+	Modify();
+	return PropertyBag->SetValueEnum(InPropertyName, InValue, Enum) == EPropertyBagResult::Success;
+}
+
+bool UMovieGraphValueView::SetValueStruct(const FName& InPropertyName, FConstStructView InValue)
+{
+	Modify();
+	return PropertyBag->SetValueStruct(InPropertyName, InValue) == EPropertyBagResult::Success;
+}
+
+bool UMovieGraphValueView::SetValueObject(const FName& InPropertyName, UObject* InValue)
+{
+	Modify();
+	return PropertyBag->SetValueObject(InPropertyName, InValue) == EPropertyBagResult::Success;
+}
+
+bool UMovieGraphValueView::SetValueClass(const FName& InPropertyName, UClass* InValue)
+{
+	Modify();
+	return PropertyBag->SetValueClass(InPropertyName, InValue) == EPropertyBagResult::Success;
+}
+
+bool UMovieGraphValueView::SetValueSerializedString(const FName& InPropertyName, const FString& NewValue)
+{
+	Modify();
+	return PropertyBag->SetValueSerializedString(InPropertyName, NewValue) == EPropertyBagResult::Success;
+}
+
+EMovieGraphValueType UMovieGraphValueView::GetValueType(const FName& InPropertyName) const
+{
+	if (const FPropertyBagPropertyDesc* Desc = PropertyBag->FindPropertyDescByName(InPropertyName))
+	{
+		return static_cast<EMovieGraphValueType>(Desc->ValueType);
+	}
+
+	return EMovieGraphValueType::None;
+}
+
+void UMovieGraphValueView::SetValueType(const FName& InPropertyName, EMovieGraphValueType ValueType, UObject* InValueTypeObject)
+{
+	if (const FPropertyBagPropertyDesc* Desc = PropertyBag->FindPropertyDescByName(InPropertyName))
+	{
+		Modify();
+
+		constexpr bool bOverwrite = true;
+		PropertyBag->AddProperty(InPropertyName, static_cast<EPropertyBagPropertyType>(ValueType), InValueTypeObject, bOverwrite);
+	}
+}
+
+const UObject* UMovieGraphValueView::GetValueTypeObject(const FName& InPropertyName) const
+{
+	if (const FPropertyBagPropertyDesc* Desc = PropertyBag->FindPropertyDescByName(InPropertyName))
+	{
+		return Desc->ValueTypeObject;
+	}
+	
+	return nullptr;
+}
+
+void UMovieGraphValueView::SetValueTypeObject(const FName& InPropertyName, const UObject* ValueTypeObject)
+{
+	if (const FPropertyBagPropertyDesc* Desc = PropertyBag->FindPropertyDescByName(InPropertyName))
+	{
+		Modify();
+
+		constexpr bool bOverwrite = true;
+
+		if (Desc->ContainerTypes.IsEmpty())
+		{
+			PropertyBag->AddProperty(InPropertyName, Desc->ValueType, ValueTypeObject, bOverwrite);
+		}
+		else
+		{
+			PropertyBag->AddContainerProperty(
+				InPropertyName, Desc->ContainerTypes, Desc->ValueType, ValueTypeObject, bOverwrite, Desc->KeyType, Desc->KeyTypeObject);
+		}
+	}
+}
+
+EMovieGraphContainerType UMovieGraphValueView::GetValueContainerType(const FName& InPropertyName) const
+{
+	if (const FPropertyBagPropertyDesc* Desc = PropertyBag->FindPropertyDescByName(InPropertyName))
+	{
+		return static_cast<EMovieGraphContainerType>(Desc->ContainerTypes.GetFirstContainerType());
+	}
+
+	return EMovieGraphContainerType::None;
+}
+
+void UMovieGraphValueView::SetValueContainerType(const FName& InPropertyName, EMovieGraphContainerType ContainerType)
+{
+	if (const FPropertyBagPropertyDesc* Desc = PropertyBag->FindPropertyDescByName(InPropertyName))
+	{
+		Modify();
+
+		constexpr bool bOverwrite = true;
+		PropertyBag->AddContainerProperty(InPropertyName, static_cast<EPropertyBagContainerType>(ContainerType), Desc->ValueType, Desc->ValueTypeObject, bOverwrite);
+	}
+}
+
+TValueOrError<FPropertyBagArrayRef, EPropertyBagResult> UMovieGraphValueView::GetArrayRef(const FName& InPropertyName)
+{
+	return PropertyBag->GetMutableArrayRef(InPropertyName);
+}
+
+// ---------
+
+UMovieGraphFixedValueView::UMovieGraphFixedValueView()
+{
+	bAllowPropertyAddRemove = false;
+}
+
+UMovieGraphFixedValueView::UMovieGraphFixedValueView(FInstancedPropertyBag* InPropertyBag)
+	: Super(InPropertyBag)
+{
+	bAllowPropertyAddRemove = false;
+}
+
+UMovieGraphMutableValueView::UMovieGraphMutableValueView()
+{
+	bAllowPropertyAddRemove = true;
+}
+
+UMovieGraphMutableValueView::UMovieGraphMutableValueView(FInstancedPropertyBag* InPropertyBag)
+	: Super(InPropertyBag)
+{
+	bAllowPropertyAddRemove = true;
+}
\ No newline at end of file
diff --git a/Engine/Plugins/MovieScene/MovieRenderPipeline/Source/MovieRenderPipelineCore/Public/Graph/MovieGraphValueView.h b/Engine/Plugins/MovieScene/MovieRenderPipeline/Source/MovieRenderPipelineCore/Public/Graph/MovieGraphValueView.h
new file mode 100644
--- /dev/null
+++ b/Engine/Plugins/MovieScene/MovieRenderPipeline/Source/MovieRenderPipelineCore/Public/Graph/MovieGraphValueView.h
@@ -0,0 +1,358 @@
+// Copyright Epic Games, Inc. All Rights Reserved.
+
+#pragma once
+
+#include "CoreMinimal.h"
+#include "MovieGraphCommon.h"
+#include "StructUtils/PropertyBag.h"
+#include "Templates/ValueOrError.h"
+
+#include "MovieGraphValueView.generated.h"
+
+#define UE_API MOVIERENDERPIPELINECORE_API
+
+namespace UE::MovieGraph::Private
+{
+	template<typename ReturnType>
+	bool GetOptionalValue(TValueOrError<ReturnType, EPropertyBagResult>& PropertyBagValue, ReturnType& OutValue)
+	{
+		// Convert the property bag-provided TValueOrError, which contains an EPropertyBagResult, to an output value and
+		// a bool (signifying if there was an error or not). EPropertyBagResult shouldn't be exposed on the MRQ API.
+		if (PropertyBagValue.HasValue())
+		{
+			OutValue = PropertyBagValue.StealValue();
+			return true;
+		}
+
+		return false;
+	}
+}
+
+/** Represents a property that should be added to a property bag via UMovieGraphValueView. */
+USTRUCT(BlueprintType)
+struct FMovieGraphValueViewProperty
+{
+	GENERATED_BODY()
+
+	FMovieGraphValueViewProperty() = default;
+
+	FMovieGraphValueViewProperty(
+		const FName& InPropertyName,
+		const EMovieGraphValueType ValueType,
+		const TObjectPtr<UObject>& InValueTypeObject,
+		const TMap<FName, FString>& InMetadata);
+
+	/** The name of the property as it will appear in the property bag. */
+	UPROPERTY(BlueprintReadWrite, Category = "Property")
+	FName PropertyName;
+
+	/** The type of value the property holds. */
+	UPROPERTY(BlueprintReadWrite, Category = "Property")
+	EMovieGraphValueType ValueType;
+
+	/** The value type object for object, struct, and enum-typed properties (nullptr for other types). */
+	UPROPERTY(BlueprintReadWrite, Category = "Property")
+	TObjectPtr<UObject> ValueTypeObject = nullptr;
+
+	/** Metadata to apply to the property. */
+	UPROPERTY(BlueprintReadWrite, Category = "Property")
+	TMap<FName, FString> Metadata;
+};
+
+/**
+ * Provides a scripting-friendly way to get/set values, and add/remove properties, in a property bag.
+ *
+ * Does not protect against invalid property bags. When using, ensure that the property bag that the view is initialized with is valid. The view
+ * will also not take ownership of the property bag.
+ */
+UCLASS(Abstract)
+class UMovieGraphValueView : public UObject
+{
+	GENERATED_BODY()
+
+public:
+	UMovieGraphValueView() = default; 
+
+	UE_API explicit UMovieGraphValueView(FInstancedPropertyBag* InPropertyBag);
+
+	/** Update the view to use the specified property bag. */
+	UE_API void ViewProperties(FInstancedPropertyBag* InPropertyBag);
+
+	/** Determines if this view is valid (meaning it points to a non-null property bag, and the bag is valid itself). */
+	UFUNCTION(BlueprintCallable, Category = "Properties")
+	UE_API bool IsValid() const;
+
+	/**
+	 * Adds a new property to the container. If a property by the given name already exists, no add operation will occur. Returns true
+	 * on successful add, otherwise returns false.
+	 */
+	UFUNCTION(BlueprintCallable, Category = "Properties")
+	UE_API bool AddProperty(const FName& InPropertyName, const EMovieGraphValueType ValueType, UObject* InValueTypeObject, const TMap<FName, FString>& InMetadata);
+
+	/**
+	 * Bulk adds properties to the container. Much more efficient than calling AddProperty() in a loop. Returns true if all properties were
+	 * successfully added, else false.
+	 */
+	UFUNCTION(BlueprintCallable, Category = "Properties")
+	UE_API bool AddProperties(const TArray<FMovieGraphValueViewProperty>& InProperties);
+
+	/** Removes a property by the given name. Returns true on successful remove, otherwise returns false. */
+	UFUNCTION(BlueprintCallable, Category = "Properties")
+	UE_API bool RemoveProperty(const FName& InPropertyName);
+
+	/** Removes all properties and their values, effectively resetting the property bag. */
+	UFUNCTION(BlueprintCallable, Category = "Properties")
+	UE_API void RemoveAllProperties();
+
+	/** Returns true if a property by the given name exists in the container, else false. */
+	UFUNCTION(BlueprintCallable, Category = "Properties")
+	UE_API bool HasProperty(const FName& InPropertyName) const;
+
+	/** Gets the number of properties that this container holds. */
+	UFUNCTION(BlueprintCallable, Category = "Properties")
+	UE_API int32 GetNumProperties() const;
+
+	/** Gets the names of properties in this container. */
+	UFUNCTION(BlueprintCallable, Category = "Properties")
+	UE_API TArray<FName> GetPropertyNames() const;
+
+#if WITH_EDITOR
+	/** Gets the metadata for the given property. */
+	UFUNCTION(BlueprintCallable, Category = "Properties")
+	UE_API TMap<FName, FString> GetPropertyMetadata(const FName& InPropertyName) const;
+#endif	// WITH_EDITOR
+
+	/** Gets the bool value of the held data. Returns true on success, else false. */
+	UFUNCTION(BlueprintCallable, Category = "Properties")
+	UE_API bool GetValueBool(const FName& InPropertyName, bool& bOutValue) const;
+
+	/** Gets the byte value of the held data. Returns true on success, else false. */
+	UFUNCTION(BlueprintCallable, Category = "Properties")
+	UE_API bool GetValueByte(const FName& InPropertyName, uint8& OutValue) const;
+
+	/** Gets the int32 value of the held data. Returns true on success, else false. */
+	UFUNCTION(BlueprintCallable, Category = "Properties")
+	UE_API bool GetValueInt32(const FName& InPropertyName, int32& OutValue) const;
+
+	/** Gets the int64 value of the held data. Returns true on success, else false. */
+	UFUNCTION(BlueprintCallable, Category = "Properties")
+	UE_API bool GetValueInt64(const FName& InPropertyName, int64& OutValue) const;
+
+	/** Gets the float value of the held data. Returns true on success, else false. */
+	UFUNCTION(BlueprintCallable, Category = "Properties")
+	UE_API bool GetValueFloat(const FName& InPropertyName, float& OutValue) const;
+
+	/** Gets the double value of the held data. Returns true on success, else false. */
+	UFUNCTION(BlueprintCallable, Category = "Properties")
+	UE_API bool GetValueDouble(const FName& InPropertyName, double& OutValue) const;
+
+	/** Gets the FName value of the held data. Returns true on success, else false. */
+	UFUNCTION(BlueprintCallable, Category = "Properties")
+	UE_API bool GetValueName(const FName& InPropertyName, FName& OutValue) const;
+
+	/** Gets the FString value of the held data. Returns true on success, else false. */
+	UFUNCTION(BlueprintCallable, Category = "Properties")
+	UE_API bool GetValueString(const FName& InPropertyName, FString& OutValue) const;
+
+	/** Gets the FText value of the held data. Returns true on success, else false. */
+	UFUNCTION(BlueprintCallable, Category = "Properties")
+	UE_API bool GetValueText(const FName& InPropertyName, FText& OutValue) const;
+
+	/** Gets the enum value (for a specific enum) of the held data. Returns true on success, else false. */
+	UFUNCTION(BlueprintCallable, Category = "Properties")
+	UE_API bool GetValueEnum(const FName& InPropertyName, uint8& OutValue, const UEnum* RequestedEnum = nullptr) const;
+
+	/** Gets the struct value (for a specific struct) of the held data. Returns true on success, else false. */
+	UE_API bool GetValueStruct(const FName& InPropertyName, FStructView& OutValue, const UScriptStruct* RequestedStruct = nullptr) const;
+
+	/** Gets the object value (for a specific class) of the held data. Returns true on success, else false. */
+	UFUNCTION(BlueprintCallable, Category = "Properties")
+	UE_API bool GetValueObject(const FName& InPropertyName, UObject*& OutValue, const UClass* RequestedClass = nullptr) const;
+
+	/** Gets the UClass value of the held data. Returns true on success, else false. */
+	UFUNCTION(BlueprintCallable, Category = "Properties")
+	UE_API bool GetValueClass(const FName& InPropertyName, UClass*& OutValue) const;
+
+	/** Gets the serialized string value of the held data. */
+	UFUNCTION(BlueprintCallable, Category = "Properties")
+	UE_API FString GetValueSerializedString(const FName& InPropertyName) const;
+
+	/** Gets the enum value of the held data. Returns true on success, else false. */
+	template <typename T>
+	bool GetValueEnum(const FName& InPropertyName, T& OutValue) const
+	{
+		TValueOrError<T, EPropertyBagResult> Result = PropertyBag->GetValueEnum<T>(InPropertyName);
+		return UE::MovieGraph::Private::GetOptionalValue<T>(Result, OutValue);
+	}
+
+	/** Gets the struct value of the held data. Returns true on success, else false. */
+	template <typename T>
+	bool GetValueStruct(const FName& InPropertyName, T* OutValue) const
+	{
+		TValueOrError<T*, EPropertyBagResult> Result = PropertyBag->GetValueStruct<T*>(InPropertyName);
+		return UE::MovieGraph::Private::GetOptionalValue<T*>(Result, OutValue);
+	}
+
+	/** Gets the object value of the held data. Returns true on success, else false. */
+	template <typename T>
+	bool GetValueObject(const FName& InPropertyName, T*& OutValue) const
+	{
+		TValueOrError<T*, EPropertyBagResult> Result = PropertyBag->GetValueObject<T*>(InPropertyName);
+		return UE::MovieGraph::Private::GetOptionalValue<T*>(Result, OutValue);
+	}
+
+	/** Sets the bool value of the held data. Returns true on success, else false. */
+	UFUNCTION(BlueprintCallable, Category = "Properties")
+	UE_API bool SetValueBool(const FName& InPropertyName, const bool bInValue);
+
+	/** Sets the byte value of the held data. Returns true on success, else false. */
+	UFUNCTION(BlueprintCallable, Category = "Properties")
+	UE_API bool SetValueByte(const FName& InPropertyName, const uint8 InValue);
+
+	/** Sets the int32 value of the held data. Returns true on success, else false. */
+	UFUNCTION(BlueprintCallable, Category = "Properties")
+	UE_API bool SetValueInt32(const FName& InPropertyName, const int32 InValue);
+
+	/** Sets the int64 value of the held data. Returns true on success, else false. */
+	UFUNCTION(BlueprintCallable, Category = "Properties")
+	UE_API bool SetValueInt64(const FName& InPropertyName, const int64 InValue);
+
+	/** Sets the float value of the held data. Returns true on success, else false. */
+	UFUNCTION(BlueprintCallable, Category = "Properties")
+	UE_API bool SetValueFloat(const FName& InPropertyName, const float InValue);
+
+	/** Sets the double value of the held data. Returns true on success, else false. */
+	UFUNCTION(BlueprintCallable, Category = "Properties")
+	UE_API bool SetValueDouble(const FName& InPropertyName, const double InValue);
+
+	/** Sets the FName value of the held data. Returns true on success, else false. */
+	UFUNCTION(BlueprintCallable, Category = "Properties")
+	UE_API bool SetValueName(const FName& InPropertyName, const FName InValue);
+
+	/** Sets the FString value of the held data. Returns true on success, else false. */
+	UFUNCTION(BlueprintCallable, Category = "Properties")
+	UE_API bool SetValueString(const FName& InPropertyName, const FString& InValue);
+
+	/** Sets the FText value of the held data. Returns true on success, else false. */
+	UFUNCTION(BlueprintCallable, Category = "Properties")
+	UE_API bool SetValueText(const FName& InPropertyName, const FText& InValue);
+
+	/** Sets the enum value of the held data. Returns true on success, else false. */
+	UFUNCTION(BlueprintCallable, Category = "Properties")
+	UE_API bool SetValueEnum(const FName& InPropertyName, const uint8 InValue, const UEnum* Enum);
+
+	/** Sets the struct value of the held data. Returns true on success, else false. */
+	UE_API bool SetValueStruct(const FName& InPropertyName, FConstStructView InValue);
+
+	/** Sets the object value of the held data. Returns true on success, else false. */
+	UFUNCTION(BlueprintCallable, Category = "Properties")
+	UE_API bool SetValueObject(const FName& InPropertyName, UObject* InValue);
+
+	/** Sets the class value of the held data. Returns true on success, else false. */
+	UFUNCTION(BlueprintCallable, Category = "Properties")
+	UE_API bool SetValueClass(const FName& InPropertyName, UClass* InValue);
+
+	/** Sets the serialized value of the held data. The string should be the serialized representation of the value. Returns true on success, else false. */
+	UFUNCTION(BlueprintCallable, Category = "Properties")
+	UE_API bool SetValueSerializedString(const FName& InPropertyName, const FString& NewValue);
+
+	/** Sets the enum value of the held data. Returns true on success, else false. */
+	template <typename T>
+	bool SetValueEnum(const FName& InPropertyName, const T InValue)
+	{
+		return PropertyBag->SetValueEnum<T>(InPropertyName, InValue) == EPropertyBagResult::Success;
+	}
+
+	/** Sets the struct value of the held data. Returns true on success, else false. */
+	template <typename T>
+	bool SetValueStruct(const FName& InPropertyName, const T& InValue)
+	{
+		return PropertyBag->SetValueStruct<T>(InPropertyName, InValue) == EPropertyBagResult::Success;
+	}
+
+	/** Sets the object value of the held data. Returns true on success, else false. */
+	template <typename T>
+	bool SetValueObject(const FName& InPropertyName, T* InValue)
+	{
+		return PropertyBag->SetValueObject<T>(InPropertyName, InValue) == EPropertyBagResult::Success;
+	}
+
+	/** Gets the type of the stored data. */
+	UFUNCTION(BlueprintCallable, Category = "Properties")
+	UE_API EMovieGraphValueType GetValueType(const FName& InPropertyName) const;
+
+	/** Sets the type of the stored data. Enums, structs, and classes must specify a value type object. */
+	UFUNCTION(BlueprintCallable, Category = "Properties")
+	UE_API void SetValueType(const FName& InPropertyName, EMovieGraphValueType ValueType, UObject* InValueTypeObject = nullptr);
+
+	/** Gets the object that defines the enum, struct, or class. */
+	UFUNCTION(BlueprintCallable, Category = "Properties")
+	UE_API const UObject* GetValueTypeObject(const FName& InPropertyName) const;
+
+	/** Sets the object that defines the enum, struct, or class. */
+	UFUNCTION(BlueprintCallable, Category = "Properties")
+	UE_API void SetValueTypeObject(const FName& InPropertyName, const UObject* ValueTypeObject);
+
+	/** Gets the container type of the stored value. */
+	UFUNCTION(BlueprintCallable, Category = "Properties")
+	UE_API EMovieGraphContainerType GetValueContainerType(const FName& InPropertyName) const;
+
+	/** Sets the container type of the stored value. */
+	UFUNCTION(BlueprintCallable, Category = "Properties")
+	UE_API void SetValueContainerType(const FName& InPropertyName, EMovieGraphContainerType ContainerType);
+
+	/**
+	 * Gets a reference to the array backing the value, if any. GetValueContainerType() will return
+	 * EMovieGraphContainerType::Array if the value is holding an array.
+	 */
+	UE_API TValueOrError<FPropertyBagArrayRef, EPropertyBagResult> GetArrayRef(const FName& InPropertyName);
+
+protected:
+	/** Whether this view should allow properties to be added/removed from the property bag. */
+	bool bAllowPropertyAddRemove = false;
+
+private:
+	FInstancedPropertyBag* PropertyBag = nullptr;
+};
+
+/**
+ * Provides a scripting-friendly way to get/set values in a property bag. Does not allow adding/removing properties.
+ */
+UCLASS(BlueprintType)
+class UMovieGraphFixedValueView final : public UMovieGraphValueView
+{
+	GENERATED_BODY()
+
+public:
+	UMovieGraphFixedValueView();
+
+	UE_API explicit UMovieGraphFixedValueView(FInstancedPropertyBag* InPropertyBag);
+
+	// The fixed view doesn't allow for adding/removing properties. Deleting these methods isn't foolproof, since a user could cast to the base class
+	// and call these methods there. However, these methods also have internal logic that cue off bAllowPropertyAddRemove to prevent these methods
+	// from performing any work.
+	bool AddProperty(
+		const FName& InPropertyName,
+		const EMovieGraphValueType ValueType,
+		UObject* InValueTypeObject,
+		const TMap<FName, FString>& InMetadata) = delete;
+	bool AddProperties(const TArray<FMovieGraphValueViewProperty>& InProperties) = delete;
+	bool RemoveProperty(const FName& InPropertyName) = delete;
+	void RemoveAllProperties() = delete;
+};
+
+/**
+ * Provides a scripting-friendly way to get/set values, as well as add/remove properties, in a property bag.
+ */
+UCLASS(BlueprintType)
+class UMovieGraphMutableValueView final : public UMovieGraphValueView
+{
+	GENERATED_BODY()
+
+public:
+	UMovieGraphMutableValueView();
+
+	UE_API explicit UMovieGraphMutableValueView(FInstancedPropertyBag* InPropertyBag);
+};
+
+#undef UE_API
\ No newline at end of file
diff --git a/Engine/Plugins/MovieScene/MovieRenderPipeline/Source/MovieRenderPipelineCore/Public/Graph/Nodes/MovieGraphClassPropertyModifier.h b/Engine/Plugins/MovieScene/MovieRenderPipeline/Source/MovieRenderPipelineCore/Public/Graph/Nodes/MovieGraphClassPropertyModifier.h
new file mode 100644
--- /dev/null
+++ b/Engine/Plugins/MovieScene/MovieRenderPipeline/Source/MovieRenderPipelineCore/Public/Graph/Nodes/MovieGraphClassPropertyModifier.h
@@ -0,0 +1,188 @@
+// Copyright Epic Games, Inc. All Rights Reserved.
+
+#pragma once
+
+#include "CoreMinimal.h"
+#include "Graph/MovieGraphRenderLayerSubsystem.h"
+#include "Graph/MovieGraphMultiValueContainer.h"
+#include "StructUtils/PropertyBag.h"
+
+#include "MovieGraphClassPropertyModifier.generated.h"
+
+#define UE_API MOVIERENDERPIPELINECORE_API
+
+/** A scripting-friendly representation of a property on an AActor or UActorComponent. */
+USTRUCT(BlueprintType)
+struct FMovieGraphPropertyReference
+{
+	GENERATED_BODY()
+
+public:
+	FMovieGraphPropertyReference() = default;
+
+	/**
+	 * Providing the FProperty is optional. Providing it just bypasses the lookup of the FProperty that will occur if only the property name and
+	 * class/component are provided.
+	 */
+	FMovieGraphPropertyReference(
+		const TObjectPtr<UClass>& InPropertyClass,
+		const TSubclassOf<UActorComponent>& InComponentClass,
+		const FName& InPropertyName,
+		FProperty* InProperty = nullptr);
+
+	/** Gets the value type of the referenced property. */
+	EMovieGraphValueType GetValueType() const;
+
+	/** Gets the value type object of the referenced property (for objects, structs, and enums). */
+	UObject* GetValueTypeObject() const;
+
+	/** Gets the FProperty that backs this property reference. */
+	FProperty* GetProperty() const;
+
+	/** Gets the name of the property that this reference refers to. */
+	FName GetPropertyName() const;
+
+	/** Gets the actor class that this property exists in. If the property exists in a component, this is the class that the component belongs to. */
+	TSubclassOf<AActor> GetPropertyClass() const;
+
+	/** Gets the class of the component that the property exists in. Can be nullptr if the property does not exist in a component. */
+	TSubclassOf<UActorComponent> GetComponentClass() const;
+
+	/** Determines if this property reference points to a valid property. */
+	bool IsValidProperty() const;
+
+	// Called by TStructOpsTypeTraits
+	void PostSerialize(const FArchive& Ar);
+
+public:
+	/**
+	 * An additional property name which may be used for other purposes. Does not affect any internal FMovieGraphPropertyReference functionality and
+	 * can be left empty if not needed.
+	 */
+	UPROPERTY(BlueprintReadWrite, Category = "Property")
+	FName AuxPropertyName;
+
+private:
+	/** The actor class that the property exists in. If the property exists within a component, this is the class that the component exists in. */
+	UPROPERTY(BlueprintReadOnly, Category = "Property", meta = (AllowPrivateAccess = "true"))
+	TSubclassOf<AActor> PropertyClass = nullptr;
+
+	/** The UActorComponent class that the property exists in, or nullptr if the property does not exist in a component. */
+	UPROPERTY(BlueprintReadOnly, Category = "Property", meta = (AllowPrivateAccess = "true"))
+	TSubclassOf<UActorComponent> ComponentClass = nullptr;
+
+	/** The name of the property within either the class or the component class. */
+	UPROPERTY(BlueprintReadOnly, Category = "Property", meta = (AllowPrivateAccess = "true"))
+	FName PropertyName;
+
+	/** Calculated (or optionally provided): The resolved FProperty that lives on the PropertyClass or ComponentClass. */
+	FProperty* Property = nullptr;
+
+	/** Calculated: The value type of the referenced property. */
+	UPROPERTY(BlueprintReadOnly, Category = "Property", meta = (AllowPrivateAccess = "true"))
+	EMovieGraphValueType ValueType = EMovieGraphValueType::None;
+
+	/** Calculated: The value type object of the referenced property (for objects, structs, and enums). */
+	UPROPERTY(BlueprintReadOnly, Category = "Property", meta = (AllowPrivateAccess = "true"))
+	TObjectPtr<UObject> ValueTypeObject = nullptr;
+
+private:
+	/** Calculates and caches property-related information so it doesn't need to be constantly re-derived. */
+	void PopulateDerivedData();
+};
+
+template<>
+struct TStructOpsTypeTraits<FMovieGraphPropertyReference> : public TStructOpsTypeTraitsBase2<FMovieGraphPropertyReference>
+{
+	enum
+	{
+		WithPostSerialize = true,
+	};
+};
+
+/**
+ * Modifies properties on specific actor type(s).
+ *
+ * The properties that should be modified are added via MovieGraphPropertyReference objects via AddProperties(). Note that the
+ * properties are matched to their EXACT AActor and/or UActorComponent class within the property reference (ie, subclass checks
+ * are not made).
+ */
+UCLASS(MinimalAPI, BlueprintType)
+class UMovieGraphClassPropertyModifier final : public UMovieGraphCollectionModifier
+{
+	GENERATED_BODY()
+
+public:
+	UMovieGraphClassPropertyModifier();
+
+	//~ Begin UMovieGraphModifierBase Interface
+	UE_API virtual void ApplyModifier(const UWorld* World) override;
+	UE_API virtual void UndoModifier() override;
+	UE_API virtual FText GetModifierName() override;
+	//~ End UMovieGraphModifierBase Interface
+
+	/**
+	 * Scripting helper that makes unreal.MovieGraphPropertyReference instances for use with this node. Little use outside of scripting. See
+	 * FMovieGraphPropertyReference documentation for further details about the property reference.
+	 */
+	UFUNCTION(BlueprintCallable, Category = "Properties")
+	static UE_API FMovieGraphPropertyReference MakePropertyReference(
+		UClass* InPropertyClass,
+		const TSubclassOf<UActorComponent>& InComponentClass,
+		const FName& InPropertyName);
+
+	/** Adds the given properties (which are each associated with a specific class). */
+	UFUNCTION(BlueprintCallable, Category = "Properties")
+	UE_API void AddProperties(const TArray<FMovieGraphPropertyReference>& InProperties);
+
+	/** Gets the property name that needs to be used with GetPropertiesForSettingValues() to set a "Custom" property value. */
+	UFUNCTION(BlueprintCallable, Category = "Properties")
+	static UE_API FName GetPropertyNameForSettingCustomValue(const FMovieGraphPropertyReference& InPropertyReference);
+
+	/**
+	 * Gets a value view which allows property values to be set. Properties cannot be added/removed with this view.
+	 *
+	 * Note that in order to modify a property value, you must use the property name provided by GetPropertyNameForSettingCustomValue().
+	 */
+	UFUNCTION(BlueprintCallable, Category = "Properties")
+	UE_API UMovieGraphFixedValueView* GetPropertiesForSettingValues();
+
+private:
+	/**
+	 * Gets the unique internal property name for the given property reference. This property name is used to track the property within the
+	 * internal property bags.
+	 */
+	static FName GetInternalPropertyName(const FMovieGraphPropertyReference& InPropertyReference);
+
+private:
+	/** All of the properties that will be modified by this modifier. */
+	TMap<TSubclassOf<AActor>, TArray<FMovieGraphPropertyReference>> PropertiesToModifyByActorClass;
+
+	/** Holds all of the properties and their associated values which will be set when the modifier is applied. */
+	UPROPERTY(Transient)
+	FInstancedPropertyBag UpdatedPropertyValues;
+
+	/** Holds all of the values of the properties prior to modification. */
+	UPROPERTY(Transient)
+	FInstancedPropertyBag PreviousPropertyValues;
+
+	/** View for the UpdatedPropertyValues property bag. */
+	UPROPERTY(Transient)
+	TObjectPtr<UMovieGraphMutableValueView> UpdatedPropertyValuesView = nullptr;
+
+	/** View for the PreviousPropertyValues property bag. */
+	UPROPERTY(Transient)
+	TObjectPtr<UMovieGraphMutableValueView> PreviousPropertyValuesView = nullptr;
+
+	/**
+	 * The property view that's user-facing and does not allow adding/removing properties (so this node can enforce how properties
+	 * are added/removed internally). Created lazily via GetPropertiesForSettingValues().
+	 */
+	UPROPERTY(Transient)
+	TObjectPtr<UMovieGraphFixedValueView> UserFacingPropertiesView = nullptr;
+
+	/** The actors and their properties that had values modified. Cached so the actor does not have to be re-resolved when undoing the modifier. */
+	TArray<TPair<UObject*, FMovieGraphPropertyReference>> ModifiedProperties;
+};
+
+#undef UE_API
\ No newline at end of file
diff --git a/Engine/Plugins/MovieScene/MovieRenderPipeline/Source/MovieRenderPipelineCore/Public/Graph/Nodes/MovieGraphLightModifierNode.h b/Engine/Plugins/MovieScene/MovieRenderPipeline/Source/MovieRenderPipelineCore/Public/Graph/Nodes/MovieGraphLightModifierNode.h
new file mode 100644
--- /dev/null
+++ b/Engine/Plugins/MovieScene/MovieRenderPipeline/Source/MovieRenderPipelineCore/Public/Graph/Nodes/MovieGraphLightModifierNode.h
@@ -0,0 +1,348 @@
+﻿// Copyright Epic Games, Inc. All Rights Reserved.
+
+#pragma once
+
+#include "Engine/Scene.h"
+#include "Graph/MovieGraphRenderLayerSubsystem.h"
+#include "MovieGraphClassPropertyModifier.h"
+#include "MovieGraphModifierNode.h"
+
+#include "MovieGraphLightModifierNode.generated.h"
+
+class UMovieGraphClassPropertyModifier;
+
+#define UE_API MOVIERENDERPIPELINECORE_API
+
+/**
+ * Controls which intensity properties are used to control lighting intensity in the Light Modifier node.
+ */
+UENUM(BlueprintType)
+enum class EMovieGraphLightModifierIntensityMethod : uint8
+{
+	/** Apply the "Intensity" property to Point, Rect, and Spot lights. */
+	PointRectSpot	UMETA(DisplayName = "Point, Rect, and Spot Lights"),
+
+	/** The main "Intensity" property is ignored, and the type-specific "Intensity" properties are used instead. */
+	PerLightActor	UMETA(DisplayName = "Per Light Actor Type")
+};
+
+/** 
+ * A node which modifies properties on light actors within the world.
+ *
+ * There are several common light properties (eg, Light Color) that can be modified directly. For less common light properties, they
+ * can be modified by adding them as "Custom" properties. The common properties apply to all light types; custom properties apply only
+ * to a specific light type.
+ */
+UCLASS(MinimalAPI)
+class UMovieGraphLightModifierNode final : public UMovieGraphSettingNode, public IMovieGraphModifierNodeInterface
+{
+	GENERATED_BODY()
+
+public:
+	UE_API UMovieGraphLightModifierNode();
+
+#if WITH_EDITOR
+	//~ Begin UMovieGraphNode interface
+	UE_API virtual FText GetNodeTitle(const bool bGetDescriptive = false) const override;
+	UE_API virtual FText GetMenuCategory() const override;
+	UE_API virtual FLinearColor GetNodeTitleColor() const override;
+	UE_API virtual FSlateIcon GetIconAndTint(FLinearColor& OutColor) const override;
+	//~ End UMovieGraphNode interface
+
+	//~ Begin UObject interface
+	UE_API virtual void PostEditChangeProperty(FPropertyChangedEvent& PropertyChangedEvent) override;
+	//~ End UObject interface
+#endif
+
+	//~ Begin UMovieGraphNode interface
+	UE_API virtual TArray<FPropertyBagPropertyDesc> GetDynamicPropertyDescriptions() const override;
+	UE_API virtual TArray<FMovieGraphPropertyInfo> GetOverrideablePropertyInfo() const override;
+	//~ End UMovieGraphNode interface
+
+	//~ Begin IMovieGraphModifierNodeInterface interface
+	UE_API virtual TArray<UMovieGraphModifierBase*> GetAllModifiers() const override;
+	UE_API virtual bool SupportsCollections() const override;
+	UE_API virtual TArray<FName> GetAllCollections() const override;
+	UE_API virtual void AddCollection(const FName& InCollectionName) override;
+	UE_API virtual bool RemoveCollection(const FName& InCollectionName) override;
+	//~ End IMovieGraphModifierNodeInterface interface
+
+	//~ Begin UMovieGraphSettingNode interface
+	UE_API virtual FString GetNodeInstanceName() const override;
+	UE_API virtual void PrepareForFlattening(const UMovieGraphSettingNode* InSourceNode) override;
+	//~ End UMovieGraphSettingNode interface
+
+	/**
+	 * Adds the property from the specified class to the "Custom" lighting properties. Does nothing if the Custom properties already contains it.
+	 * Returns true on success, else false.
+	 *
+	 * Adding a custom property does not immediately make it take effect. Its override state needs to be set to "true" via
+	 * UpdateCustomLightPropertyOverrideState().
+	 */
+	UE_API bool AddCustomLightProperty(UClass* InLightClass, const FProperty* InLightingProperty);
+
+	/**
+	 * Adds the property from the specified class to the "Custom" lighting properties. Does nothing if the Custom properties already contains it.
+	 * Returns true on success, else false.
+	 * 
+	 * Adding a custom property does not immediately make it take effect. Its override state needs to be set to "true" via
+	 * UpdateCustomLightPropertyOverrideState().
+	 */
+	UFUNCTION(BlueprintCallable, Category = "Properties")
+	UE_API bool AddCustomLightProperty(UClass* InLightClass, const FName& InPropertyName);
+
+	/**
+	 * Removes the property in the specified class from the "Custom" lighting properties. Does nothing if the Custom properties doesn't contain the
+	 * property. Returns true on success, else false.
+	 */
+	UE_API bool RemoveCustomLightProperty(const UClass* InLightClass, const FProperty* InLightingProperty);
+
+	/**
+	 * Removes the property in the specified class from the "Custom" lighting properties. Does nothing if the Custom properties doesn't contain the
+	 * property. Returns true on success, else false.
+	 */
+	UFUNCTION(BlueprintCallable, Category = "Properties")
+	UE_API bool RemoveCustomLightProperty(const UClass* InLightClass, const FName& InPropertyName);
+
+	/** Determines if the given property exists within the "Custom" lighting properties. */
+	UE_API bool HasCustomLightProperty(const UClass* InLightClass, const FProperty* InLightingProperty) const;
+	
+	/** Determines if the given property exists within the "Custom" lighting properties. */
+	UFUNCTION(BlueprintCallable, Category = "Properties")
+	UE_API bool HasCustomLightProperty(const UClass* InLightClass, const FName& InPropertyName) const;
+
+	/**
+	 * Determines if the given "Custom" light property has been marked as overridden. A custom property MUST be marked as overridden in
+	 * order to take effect. The override state is visualized as the checkbox in front of the property name the UI.
+	 */
+	UFUNCTION(BlueprintCallable, Category = "Properties")
+	UE_API bool IsCustomLightPropertyOverridden(const UClass* InLightClass, const FName& InPropertyName) const;
+	
+	/**
+	 * Updates the override state for a "Custom" light property. A custom property MUST be marked as overridden in order to take effect. Returns true
+	 * on success, else false.
+	 */
+	UFUNCTION(BlueprintCallable, Category = "Properties")
+	UE_API bool UpdateCustomLightPropertyOverrideState(const UClass* InLightClass, const FName& InPropertyName, const bool bIsOverridden);
+
+	/** Gets how many "Custom" lighting properties have been added. */
+	UFUNCTION(BlueprintCallable, Category = "Properties")
+	UE_API int32 GetNumCustomLightProperties() const;
+
+public:
+	// Always merge the modifier name, no need for the user to do this explicitly
+	UPROPERTY()
+	uint8 bOverride_ModifierName : 1 = 1;
+
+	UPROPERTY(EditAnywhere, BlueprintReadWrite, Category = Overrides, meta = (InlineEditConditionToggle))
+	uint8 bOverride_DirectionalLightIntensity : 1;
+
+	UPROPERTY(EditAnywhere, BlueprintReadWrite, Category = Overrides, meta = (InlineEditConditionToggle))
+	uint8 bOverride_SkyLightIntensityScale : 1;
+
+	UPROPERTY(EditAnywhere, BlueprintReadWrite, Category = Overrides, meta = (InlineEditConditionToggle))
+	uint8 bOverride_IntensityMethod : 1;
+
+	UPROPERTY(EditAnywhere, BlueprintReadWrite, Category = Overrides, meta = (InlineEditConditionToggle))
+	uint8 bOverride_Intensity : 1;
+
+	UPROPERTY(EditAnywhere, BlueprintReadWrite, Category = Overrides, meta = (InlineEditConditionToggle))
+	uint8 bOverride_PointLightIntensity : 1;
+
+	UPROPERTY(EditAnywhere, BlueprintReadWrite, Category = Overrides, meta = (InlineEditConditionToggle))
+	uint8 bOverride_PointLightIntensityUnits : 1;
+
+	UPROPERTY(EditAnywhere, BlueprintReadWrite, Category = Overrides, meta = (InlineEditConditionToggle))
+	uint8 bOverride_RectLightIntensity : 1;
+
+	UPROPERTY(EditAnywhere, BlueprintReadWrite, Category = Overrides, meta = (InlineEditConditionToggle))
+	uint8 bOverride_RectLightIntensityUnits : 1;
+
+	UPROPERTY(EditAnywhere, BlueprintReadWrite, Category = Overrides, meta = (InlineEditConditionToggle))
+	uint8 bOverride_SpotLightIntensity : 1;
+
+	UPROPERTY(EditAnywhere, BlueprintReadWrite, Category = Overrides, meta = (InlineEditConditionToggle))
+	uint8 bOverride_SpotLightIntensityUnits : 1;
+
+	UPROPERTY(EditAnywhere, BlueprintReadWrite, Category = Overrides, meta = (InlineEditConditionToggle))
+	uint8 bOverride_IntensityUnits : 1;
+
+	UPROPERTY(EditAnywhere, BlueprintReadWrite, Category = Overrides, meta = (InlineEditConditionToggle))
+	uint8 bOverride_LightColor : 1;
+
+	UPROPERTY(EditAnywhere, BlueprintReadWrite, Category = Overrides, meta = (InlineEditConditionToggle))
+	uint8 bOverride_bAffectsWorld : 1;
+
+	UPROPERTY(EditAnywhere, BlueprintReadWrite, Category = Overrides, meta = (InlineEditConditionToggle))
+	uint8 bOverride_CastShadows : 1;
+
+	UPROPERTY(EditAnywhere, BlueprintReadWrite, Category = Overrides, meta = (InlineEditConditionToggle))
+	uint8 bOverride_IndirectLightingIntensity : 1;
+
+	UPROPERTY(EditAnywhere, BlueprintReadWrite, Category = Overrides, meta = (InlineEditConditionToggle))
+	uint8 bOverride_VolumetricScatteringIntensity : 1;
+
+	UPROPERTY(EditAnywhere, BlueprintReadWrite, Category = Overrides, meta = (InlineEditConditionToggle))
+	uint8 bOverride_LightingChannels : 1;
+
+	UPROPERTY(EditAnywhere, BlueprintReadWrite, Category = Overrides, meta = (InlineEditConditionToggle))
+	uint8 bOverride_SamplesPerPixel : 1;
+
+	/** The name of this modifier. */
+	UPROPERTY(EditAnywhere, BlueprintReadWrite, Category = "Modifier")
+	FString ModifierName;
+
+	/** 
+	 * Maximum illumination from the light (in lux).
+	 */
+	UPROPERTY(EditAnywhere, BlueprintReadWrite, Category = "Light|Intensity", meta = (EditCondition = "bOverride_DirectionalLightIntensity", UIMin = "0.0", UIMax = "150.0", SliderExponent = "2.0", Units = "lux"))
+	float DirectionalLightIntensity = 10.f;
+
+	/** 
+	 * Total energy that the sky light emits.
+	 */
+	UPROPERTY(EditAnywhere, BlueprintReadWrite, Category = "Light|Intensity", meta = (EditCondition = "bOverride_SkyLightIntensityScale", UIMin = "0.0", UIMax = "50000.0", SliderExponent = "10.0"))
+	float SkyLightIntensityScale = 1.f;
+
+	/** 
+	 * Whether point, rect, and spot lights share the same intensity, or the intensity (and intensity units) are unique for each type.
+	 */
+	UPROPERTY(EditAnywhere, BlueprintReadWrite, Category = "Light|Intensity", meta = (EditCondition = "bOverride_IntensityMethod"))
+	EMovieGraphLightModifierIntensityMethod IntensityMethod = EMovieGraphLightModifierIntensityMethod::PointRectSpot;
+
+	/** 
+	 * Total energy that the light emits.
+	 */
+	UPROPERTY(EditAnywhere, BlueprintReadWrite, Category = "Light|Intensity", meta = (EditCondition = "bOverride_Intensity", UIMin = "0.0", UIMax = "20.0", Units = "cd"))
+	float Intensity = 8.f;
+
+	/** 
+	 * Total energy that the point light emits.
+	 */
+	UPROPERTY(EditAnywhere, BlueprintReadWrite, Category = "Light|Intensity|PointLight", meta = (EditCondition = "bOverride_PointLightIntensity", UIMin = "0.0", UIMax = "20.0"))
+	float PointLightIntensity = 8.f;
+
+	/** 
+	 * Units used for the point light intensity. 
+	 * The peak luminous intensity is measured in candelas, while the luminous flux is measured in lumens.
+	 * When the units are set in Nits, the light's power is also determined by the size of the light source (larger sources will emit more light).
+	 */
+	UPROPERTY(EditAnywhere, BlueprintReadWrite, Category = "Light|Intensity|PointLight", meta = (EditCondition = "bOverride_PointLightIntensityUnits"))
+	ELightUnits PointLightIntensityUnits = ELightUnits::Candelas;
+
+	/** 
+	 * Total energy that the rect light emits.
+	 */
+	UPROPERTY(EditAnywhere, BlueprintReadWrite, Category = "Light|Intensity|RectLight", meta = (EditCondition = "bOverride_RectLightIntensity", UIMin = "0.0", UIMax = "20.0"))
+	float RectLightIntensity = 8.f;
+
+	/** 
+	 * Units used for the rect light intensity. 
+	 * The peak luminous intensity is measured in candelas, while the luminous flux is measured in lumens.
+	 * When the units are set in Nits, the light's power is also determined by the size of the light source (larger sources will emit more light).
+	 */
+	UPROPERTY(EditAnywhere, BlueprintReadWrite, Category = "Light|Intensity|RectLight", meta = (EditCondition = "bOverride_RectLightIntensityUnits"))
+	ELightUnits RectLightIntensityUnits = ELightUnits::Candelas;
+
+	/** 
+	 * Total energy that the spot light emits.
+	 */
+	UPROPERTY(EditAnywhere, BlueprintReadWrite, Category = "Light|Intensity|SpotLight", meta = (EditCondition = "bOverride_SpotLightIntensity", UIMin = "0.0", UIMax = "20.0"))
+	float SpotLightIntensity = 8.f;
+
+	/** 
+	 * Units used for the spot light intensity. 
+	 * The peak luminous intensity is measured in candelas, while the luminous flux is measured in lumens.
+	 * When the units are set in Nits, the light's power is also determined by the size of the light source (larger sources will emit more light).
+	 */
+	UPROPERTY(EditAnywhere, BlueprintReadWrite, Category = "Light|Intensity|SpotLight", meta = (EditCondition = "bOverride_SpotLightIntensityUnits"))
+	ELightUnits SpotLightIntensityUnits = ELightUnits::Candelas;
+
+	/** 
+	 * Units used for the intensity. 
+	 * The peak luminous intensity is measured in candelas, while the luminous flux is measured in lumens.
+	 * When the units are set in Nits, the light's power is also determined by the size of the light source (larger sources will emit more light).
+	 */
+	UPROPERTY(EditAnywhere, BlueprintReadWrite, Category = "Light|Intensity", meta = (EditCondition = "bOverride_IntensityUnits"))
+	ELightUnits IntensityUnits = ELightUnits::Candelas;
+
+	/** 
+	 * Filter color of the light.
+	 * Note that this can change the light's effective intensity.
+	 */
+	UPROPERTY(EditAnywhere, BlueprintReadWrite, Category = "Light", meta = (EditCondition = "bOverride_LightColor", HideAlphaChannel))
+	FColor LightColor = FColor::White;
+
+	/** 
+	 * Whether the light can affect the world, or whether it is disabled.
+	 * A disabled light will not contribute to the scene in any way.  This setting cannot be changed at runtime and unbuilds lighting when changed.
+	 * Setting this to false has the same effect as deleting the light, so it is useful for non-destructive experiments.
+	 */
+	UPROPERTY(EditAnywhere, BlueprintReadWrite, Category = "Light", meta = (EditCondition = "bOverride_bAffectsWorld"))
+	uint32 bAffectsWorld : 1 = 1;
+
+	/**
+	 * Whether the light should cast any shadows.
+	 */
+	UPROPERTY(EditAnywhere, BlueprintReadWrite, Category = "Light", meta = (EditCondition = "bOverride_CastShadows"))
+	uint32 CastShadows : 1 = 1;
+
+	/** 
+	 * Scales the indirect lighting contribution from this light. 
+	 * A value of 0 disables any GI from this light. Default is 1.
+	 */
+	UPROPERTY(EditAnywhere, BlueprintReadWrite, Category = "Light", meta = (UIMin = "0.0", UIMax = "6.0", EditCondition = "bOverride_IndirectLightingIntensity"))
+	float IndirectLightingIntensity = 1.f;
+
+	/** Intensity of the volumetric scattering from this light.  This scales Intensity and LightColor. */
+	UPROPERTY(EditAnywhere, BlueprintReadWrite, Category = "Light", meta = (UIMin = "0.25", UIMax = "4.0", EditCondition = "bOverride_VolumetricScatteringIntensity"))
+	float VolumetricScatteringIntensity = 1.f;
+
+	/** 
+	 * Channels that this light should affect.  
+	 * These channels only apply to opaque materials, direct lighting, and dynamic lighting and shadowing.
+	 * Lighting channels are only supported on translucent materials using forward shading (i.e. when not using the translucency lighting volume).
+	 */
+	UPROPERTY(EditAnywhere, BlueprintReadWrite, Category = "Light", meta = (EditCondition = "bOverride_LightingChannels"))
+	FLightingChannels LightingChannels;
+
+	/** Samples per pixel for ray tracing. */
+	UPROPERTY(EditAnywhere, BlueprintReadWrite, Category = "RayTracing", meta = (EditCondition = "bOverride_SamplesPerPixel"))
+	int SamplesPerPixel = 1;
+
+private:
+	/** Gets the property name used for the given property within the DynamicProperties property bag. */
+	UE_API static FName GetLightingPropertyName(const UClass* InLightClass, const FProperty* InLightingProperty);
+
+	/**
+	 * Gets the bOverride_* property name used for the given property within the DynamicProperties property bag. The "main" property name
+	 * is the property name returned by GetLightingPropertyName().
+	 */
+	UE_API static FName GetLightingOverridePropertyName(const FName& InMainPropertyName);
+
+private:
+	// Always merge collection names, no need for the user to do this explicitly
+	UPROPERTY()
+	uint8 bOverride_Collections : 1 = 1;
+
+	// Same with the custom property info; always merge it
+	UPROPERTY()
+	uint8 bOverride_CustomPropertyInfo : 1 = 1;
+
+	/** The collections modified by this node. Only lights will be modified. */
+	UPROPERTY(meta = (HideInActiveRenderSettings))
+	TArray<FName> Collections;
+
+	/** The light modifier associated with this node. */
+	UPROPERTY(meta = (HideInActiveRenderSettings))
+	TObjectPtr<UMovieGraphClassPropertyModifier> Modifier;
+
+	/**
+	 * Stores supplemental information for each "Custom" lighting property. Key is the name of the corresponding non-override property
+	 * within DynamicProperties.
+	 */
+	UPROPERTY(meta = (HideInActiveRenderSettings))
+	TMap<FName, FMovieGraphPropertyReference> CustomPropertyInfo;
+};
+
+#undef UE_API
diff --git a/Engine/Plugins/MovieScene/MovieRenderPipeline/Source/MovieRenderPipelineEditor/Private/Customizations/Graph/MovieGraphLightModifierCustomization.cpp b/Engine/Plugins/MovieScene/MovieRenderPipeline/Source/MovieRenderPipelineEditor/Private/Customizations/Graph/MovieGraphLightModifierCustomization.cpp
new file mode 100644
--- /dev/null
+++ b/Engine/Plugins/MovieScene/MovieRenderPipeline/Source/MovieRenderPipelineEditor/Private/Customizations/Graph/MovieGraphLightModifierCustomization.cpp
@@ -0,0 +1,525 @@
+// Copyright Epic Games, Inc. All Rights Reserved.
+
+#include "MovieGraphLightModifierCustomization.h"
+
+#include "Components/LightComponentBase.h"
+#include "Components/LocalLightComponent.h"
+#include "DetailCategoryBuilder.h"
+#include "DetailLayoutBuilder.h"
+#include "DetailWidgetRow.h"
+#include "Engine/Light.h"
+#include "Engine/SkyLight.h"
+#include "Graph/Nodes/MovieGraphModifierNode.h"
+#include "IPropertyUtilities.h"
+#include "ScopedTransaction.h"
+#include "Styling/SlateIconFinder.h"
+#include "UObject/WeakInterfacePtr.h"
+#include "Widgets/Graph/SModifierCollectionsHeader.h"
+#include "Widgets/Graph/SMovieGraphModifierCollectionsList.h"
+
+#define LOCTEXT_NAMESPACE "MovieGraphLightModifierCustomization"
+
+TSharedRef<IDetailCustomization> FMovieGraphLightModifierCustomization::MakeInstance()
+{
+	return MakeShared<FMovieGraphLightModifierCustomization>();
+}
+
+void FMovieGraphLightModifierCustomization::CustomizeDetails(IDetailLayoutBuilder& InDetailBuilder)
+{
+	// The customization only supports editing a single Modifier node
+	const TWeakInterfacePtr<IMovieGraphModifierNodeInterface> ModifierNode = GetSelectedModifierNode();
+	if (!ModifierNode.IsValid())
+	{
+		return;
+	}
+
+	// Replace the "Collection" category row with a custom whole-row widget which includes an add-collection button
+	IDetailCategoryBuilder& CollectionsCategory =
+		InDetailBuilder.EditCategory(FName("Collections"), FText::GetEmpty(), ECategoryPriority::Uncommon);
+	CollectionsCategory.HeaderContent
+	(
+		SNew(SMovieGraphCollectionsHeaderWidget)
+		.WeakModifierInterface(ModifierNode)
+		.OnCollectionPicked_Lambda([this](const FName CollectionName)
+		{
+			CollectionsList->Refresh();
+		})
+	, /* bWholeRowContent */ true);
+
+	// Add a collections browser
+	CollectionsCategory.AddCustomRow(FText::GetEmpty())
+	.WholeRowWidget
+	[
+		SAssignNew(CollectionsList, SMovieGraphModifierCollectionsList)
+		.WeakModifierInterface(ModifierNode)
+	];
+
+	// Add a "Custom" category and include a "+" button in its header.
+	IDetailCategoryBuilder& CustomCategory =
+		InDetailBuilder.EditCategory(FName("Custom"), FText::GetEmpty(), ECategoryPriority::Uncommon);
+	CustomCategory.HeaderContent
+	(
+		SNew(SHorizontalBox)
+		+ SHorizontalBox::Slot()
+		.HAlign(HAlign_Left)
+		.VAlign(VAlign_Center)
+		.AutoWidth()
+		[
+			SNew(STextBlock)
+			.Text(LOCTEXT("CustomHeaderText", "Custom"))
+			.Font(FAppStyle::Get().GetFontStyle("DetailsView.CategoryFontStyle"))
+		]
+
+		+ SHorizontalBox::Slot()
+		.HAlign(HAlign_Left)
+		.VAlign(VAlign_Center)
+		.Padding(5.f, 0, 0, 0)
+		[
+			SNew(SComboButton)
+			.ToolTipText(LOCTEXT("AddCustomLightPropertyTooltip", "Add a light property that will be updated by this modifier."))
+			.ComboButtonStyle(FAppStyle::Get(), "SimpleComboButton")
+			.ContentPadding(0)
+			.HasDownArrow(false)
+			.OnGetMenuContent_Lambda([this, ModifierNode]()
+			{
+				return GetCustomMenuContents(ModifierNode);
+			})
+			.ButtonContent()
+			[
+				SNew(SBox)
+				.WidthOverride(16)
+				.HeightOverride(16)
+				[
+					SNew(SImage)
+					.Image(FAppStyle::Get().GetBrush("Icons.AddCircle"))
+					.ColorAndOpacity(FSlateColor::UseForeground())
+				]
+			]
+		]
+	, /* bWholeRowContent */ true);
+
+	// Move the node's DynamicProperties (used to store all of the "Custom" lighting properties) to be under the "Custom" category.
+	const TSharedRef<IPropertyHandle> LightingValuesProperty = InDetailBuilder.GetProperty(TEXT("DynamicProperties"), UMovieGraphNode::StaticClass());
+	if (LightingValuesProperty->IsValidHandle())
+	{
+		InDetailBuilder.HideProperty(LightingValuesProperty);
+		CustomCategory.AddProperty(LightingValuesProperty);
+	}
+
+	// Only add the "No Custom Properties" text if the LightingValues value container is empty.
+	if (const UMovieGraphLightModifierNode* LightModifierNode = Cast<UMovieGraphLightModifierNode>(ModifierNode.Get()))
+	{
+		if (LightModifierNode->GetNumCustomLightProperties() == 0)
+		{
+			CustomCategory.AddCustomRow(FText::GetEmpty())
+			.WholeRowWidget
+			.VAlign(VAlign_Center)
+			.HAlign(HAlign_Center)
+			[
+				SNew(SBox)
+				.Padding(15.f)
+				[
+					SNew(STextBlock)
+					.ColorAndOpacity(FSlateColor::UseSubduedForeground())
+					.Font(InDetailBuilder.GetDetailFontItalic())
+					.Text(FText::FromString(TEXT("Add a Custom Light Property")))
+				]
+			];
+		}
+	}
+
+	// Define the sort order of the categories.
+	int32 SortOrder = 0;
+	InDetailBuilder.EditCategory(FName("Modifier")).SetSortOrder(SortOrder);
+	InDetailBuilder.EditCategory(FName("Collections")).SetSortOrder(++SortOrder);
+	InDetailBuilder.EditCategory(FName("Light")).SetSortOrder(++SortOrder);
+	InDetailBuilder.EditCategory(FName("RayTracing")).SetSortOrder(++SortOrder);
+	InDetailBuilder.EditCategory(FName("Custom")).SetSortOrder(++SortOrder);
+
+	const TSharedRef<IPropertyHandle> IntensityMethodProperty = InDetailBuilder.GetProperty(GET_MEMBER_NAME_CHECKED(UMovieGraphLightModifierNode, IntensityMethod));
+	const TSharedRef<IPropertyHandle> IntensityUnitsProperty = InDetailBuilder.GetProperty(GET_MEMBER_NAME_CHECKED(UMovieGraphLightModifierNode, IntensityUnits));
+	const TSharedRef<IPropertyHandle> PointLightIntensityUnitsProperty = InDetailBuilder.GetProperty(GET_MEMBER_NAME_CHECKED(UMovieGraphLightModifierNode, PointLightIntensityUnits));
+	const TSharedRef<IPropertyHandle> RectLightIntensityUnitsProperty = InDetailBuilder.GetProperty(GET_MEMBER_NAME_CHECKED(UMovieGraphLightModifierNode, RectLightIntensityUnits));
+	const TSharedRef<IPropertyHandle> SpotLightIntensityUnitsProperty = InDetailBuilder.GetProperty(GET_MEMBER_NAME_CHECKED(UMovieGraphLightModifierNode, SpotLightIntensityUnits));
+
+	// This is a bit of a heavy hammer, but because many of the intensity-related properties are deeply nested in sub-categories, we need to force
+	// refresh the details panel when the Intensity Method changes. This is required to get the visibility lambda on the intensity-related properties
+	// to take effect. Without this, the visibility lambdas will still fire, but the layout of the details panel will not update to reflect
+	// visibility changes.
+	IntensityMethodProperty->SetOnPropertyValueChanged(FSimpleDelegate::CreateSPLambda(this, [this]()
+	{
+		if (const TSharedPtr<IDetailLayoutBuilder> DetailBuilderPin = DetailBuilder.Pin())
+		{
+			DetailBuilderPin->ForceRefreshDetails();
+		}
+	}));
+
+	// When the Intensity Units property on any of the Intensity properties is about to change, cache the current value of the Intensity Units. 
+	IntensityUnitsProperty->SetOnPropertyValuePreChange(FSimpleDelegate::CreateSP(this, &FMovieGraphLightModifierCustomization::OnIntensityUnitsValuePreChange, GET_MEMBER_NAME_CHECKED(UMovieGraphLightModifierNode, IntensityUnits)));
+	PointLightIntensityUnitsProperty->SetOnPropertyValuePreChange(FSimpleDelegate::CreateSP(this, &FMovieGraphLightModifierCustomization::OnIntensityUnitsValuePreChange, GET_MEMBER_NAME_CHECKED(UMovieGraphLightModifierNode, PointLightIntensityUnits)));
+	RectLightIntensityUnitsProperty->SetOnPropertyValuePreChange(FSimpleDelegate::CreateSP(this, &FMovieGraphLightModifierCustomization::OnIntensityUnitsValuePreChange, GET_MEMBER_NAME_CHECKED(UMovieGraphLightModifierNode, RectLightIntensityUnits)));
+	SpotLightIntensityUnitsProperty->SetOnPropertyValuePreChange(FSimpleDelegate::CreateSP(this, &FMovieGraphLightModifierCustomization::OnIntensityUnitsValuePreChange, GET_MEMBER_NAME_CHECKED(UMovieGraphLightModifierNode, SpotLightIntensityUnits)));
+
+	// When the intensity units change, the value of the intensity needs to be updated. The details panel also needs to be force-refreshed to get the
+	// units metadata to update properly.
+	IntensityUnitsProperty->SetOnPropertyValueChanged(
+		FSimpleDelegate::CreateSP(this, &FMovieGraphLightModifierCustomization::OnIntensityUnitsValueChange,
+			GET_MEMBER_NAME_CHECKED(UMovieGraphLightModifierNode, Intensity),
+			GET_MEMBER_NAME_CHECKED(UMovieGraphLightModifierNode, IntensityUnits)));
+	PointLightIntensityUnitsProperty->SetOnPropertyValueChanged(
+		FSimpleDelegate::CreateSP(this, &FMovieGraphLightModifierCustomization::OnIntensityUnitsValueChange,
+			GET_MEMBER_NAME_CHECKED(UMovieGraphLightModifierNode, PointLightIntensity),
+			GET_MEMBER_NAME_CHECKED(UMovieGraphLightModifierNode, PointLightIntensityUnits)));
+	RectLightIntensityUnitsProperty->SetOnPropertyValueChanged(
+		FSimpleDelegate::CreateSP(this, &FMovieGraphLightModifierCustomization::OnIntensityUnitsValueChange,
+			GET_MEMBER_NAME_CHECKED(UMovieGraphLightModifierNode, RectLightIntensity),
+			GET_MEMBER_NAME_CHECKED(UMovieGraphLightModifierNode, RectLightIntensityUnits)));
+	SpotLightIntensityUnitsProperty->SetOnPropertyValueChanged(
+		FSimpleDelegate::CreateSP(this, &FMovieGraphLightModifierCustomization::OnIntensityUnitsValueChange,
+			GET_MEMBER_NAME_CHECKED(UMovieGraphLightModifierNode, SpotLightIntensity),
+			GET_MEMBER_NAME_CHECKED(UMovieGraphLightModifierNode, SpotLightIntensityUnits)));
+
+	// Some of the intensity properties should not be visible depending on the intensity method. EditCondition metadata (+ EditConditionHides)
+	// can't be used because EditCondition is already used with the corresponding bOverride_* properties. Instead, use a visibility lambda.
+	ShowPropertyConditionally(GET_MEMBER_NAME_CHECKED(UMovieGraphLightModifierNode, Intensity), EMovieGraphLightModifierIntensityMethod::PointRectSpot);
+	ShowPropertyConditionally(GET_MEMBER_NAME_CHECKED(UMovieGraphLightModifierNode, IntensityUnits), EMovieGraphLightModifierIntensityMethod::PointRectSpot);
+	ShowPropertyConditionally(GET_MEMBER_NAME_CHECKED(UMovieGraphLightModifierNode, PointLightIntensity), EMovieGraphLightModifierIntensityMethod::PerLightActor);
+	ShowPropertyConditionally(GET_MEMBER_NAME_CHECKED(UMovieGraphLightModifierNode, PointLightIntensityUnits), EMovieGraphLightModifierIntensityMethod::PerLightActor);
+	ShowPropertyConditionally(GET_MEMBER_NAME_CHECKED(UMovieGraphLightModifierNode, RectLightIntensity), EMovieGraphLightModifierIntensityMethod::PerLightActor);
+	ShowPropertyConditionally(GET_MEMBER_NAME_CHECKED(UMovieGraphLightModifierNode, RectLightIntensityUnits), EMovieGraphLightModifierIntensityMethod::PerLightActor);
+	ShowPropertyConditionally(GET_MEMBER_NAME_CHECKED(UMovieGraphLightModifierNode, SpotLightIntensity), EMovieGraphLightModifierIntensityMethod::PerLightActor);
+	ShowPropertyConditionally(GET_MEMBER_NAME_CHECKED(UMovieGraphLightModifierNode, SpotLightIntensityUnits), EMovieGraphLightModifierIntensityMethod::PerLightActor);
+
+	// Initialize the metadata on the Intensity-related properties ("Units", etc -- these depend on the Intensity Units setting).
+	InitializeIntensityProperty(GET_MEMBER_NAME_CHECKED(UMovieGraphLightModifierNode, Intensity), GET_MEMBER_NAME_CHECKED(UMovieGraphLightModifierNode, IntensityUnits));
+	InitializeIntensityProperty(GET_MEMBER_NAME_CHECKED(UMovieGraphLightModifierNode, PointLightIntensity), GET_MEMBER_NAME_CHECKED(UMovieGraphLightModifierNode, PointLightIntensityUnits));
+	InitializeIntensityProperty(GET_MEMBER_NAME_CHECKED(UMovieGraphLightModifierNode, RectLightIntensity), GET_MEMBER_NAME_CHECKED(UMovieGraphLightModifierNode, RectLightIntensityUnits));
+	InitializeIntensityProperty(GET_MEMBER_NAME_CHECKED(UMovieGraphLightModifierNode, SpotLightIntensity), GET_MEMBER_NAME_CHECKED(UMovieGraphLightModifierNode, SpotLightIntensityUnits));
+}
+
+void FMovieGraphLightModifierCustomization::CustomizeDetails(const TSharedPtr<IDetailLayoutBuilder>& InDetailBuilder)
+{
+	DetailBuilder = InDetailBuilder;
+	CustomizeDetails(*InDetailBuilder);
+}
+
+TWeakInterfacePtr<IMovieGraphModifierNodeInterface> FMovieGraphLightModifierCustomization::GetSelectedModifierNode() const
+{
+	if (const TSharedPtr<IDetailLayoutBuilder> DetailBuilderPin = DetailBuilder.Pin())
+	{
+		TArray<TWeakObjectPtr<UObject>> CustomizedObjects;
+		DetailBuilderPin->GetObjectsBeingCustomized(CustomizedObjects);
+
+		for (const TWeakObjectPtr<UObject>& CustomizedObject : CustomizedObjects)
+		{
+			if (CustomizedObject.IsValid() && CustomizedObject->Implements<UMovieGraphModifierNodeInterface>())
+			{
+				return TWeakInterfacePtr<IMovieGraphModifierNodeInterface>(CustomizedObject.Get());
+			}
+		}
+	}
+
+	return nullptr;
+}
+
+void FMovieGraphLightModifierCustomization::InitializeIntensityProperty(const FName& InIntensityPropertyName, const FName& InIntensityUnitsPropertyName)
+{
+	if (const TSharedPtr<IDetailLayoutBuilder> DetailBuilderPin = DetailBuilder.Pin())
+	{
+		const TSharedRef<IPropertyHandle> IntensityHandle = DetailBuilderPin->GetProperty(InIntensityPropertyName, UMovieGraphLightModifierNode::StaticClass());
+		const TSharedRef<IPropertyHandle> IntensityUnitsHandle = DetailBuilderPin->GetProperty(InIntensityUnitsPropertyName, UMovieGraphLightModifierNode::StaticClass());
+
+		uint8 IntensityUnits = (uint8)ELightUnits::Candelas;
+		IntensityUnitsHandle->GetValue(IntensityUnits);
+		const ELightUnits IntensityUnitsEnum = static_cast<ELightUnits>(IntensityUnits);
+
+		if (IntensityUnitsEnum == ELightUnits::EV)
+		{
+			IntensityHandle->SetInstanceMetaData("UIMin",TEXT("-32.0"));
+			IntensityHandle->SetInstanceMetaData("UIMax",TEXT("32.0"));
+		}
+		else
+		{
+			const float ConversionFactor = ULocalLightComponent::GetUnitsConversionFactor((ELightUnits)0, IntensityUnitsEnum);
+	
+			IntensityHandle->SetInstanceMetaData("UIMin",TEXT("0.0"));
+			IntensityHandle->SetInstanceMetaData("UIMax",  *FString::SanitizeFloat(100000.0f * ConversionFactor));
+			IntensityHandle->SetInstanceMetaData("SliderExponent", TEXT("2.0"));
+		}
+
+		if (IntensityUnitsEnum == ELightUnits::Lumens)
+		{
+			IntensityHandle->SetInstanceMetaData("Units", TEXT("lm"));
+		}
+		else if (IntensityUnitsEnum == ELightUnits::Candelas)
+		{
+			IntensityHandle->SetInstanceMetaData("Units", TEXT("cd"));
+		}
+		else if (IntensityUnitsEnum == ELightUnits::EV)
+		{
+			IntensityHandle->SetInstanceMetaData("Units", TEXT("ev"));
+		}
+		else if (IntensityUnitsEnum == ELightUnits::Nits)
+		{
+			IntensityHandle->SetInstanceMetaData("Units", TEXT("nt"));
+		}
+		else
+		{
+			IntensityHandle->SetInstanceMetaData("Units", TEXT(""));
+		}
+	}
+}
+
+void FMovieGraphLightModifierCustomization::ShowPropertyConditionally(FName IntensityPropertyName, EMovieGraphLightModifierIntensityMethod ShowWithMethod)
+{
+	const TSharedPtr<IDetailLayoutBuilder> DetailBuilderPin = DetailBuilder.Pin();
+	if (!DetailBuilderPin.IsValid())
+	{
+		return;
+	}
+
+	const TSharedRef<IPropertyHandle> IntensityHandle = DetailBuilderPin->GetProperty(IntensityPropertyName);
+	const TSharedRef<IPropertyHandle> IntensityMethodHandle = DetailBuilderPin->GetProperty(GET_MEMBER_NAME_CHECKED(UMovieGraphLightModifierNode, IntensityMethod));
+
+	if (!IntensityHandle->IsValidHandle() || !IntensityMethodHandle->IsValidHandle())
+	{
+		return;
+	}
+
+	if (IDetailPropertyRow* IntensityPropertyRow = DetailBuilderPin->EditDefaultProperty(IntensityHandle))
+	{
+		IntensityPropertyRow->Visibility(MakeAttributeLambda([IntensityMethodHandle, ShowWithMethod]()
+		{
+			uint8 IntensityMethod = 0;
+			IntensityMethodHandle->GetValue(IntensityMethod);
+
+			const EVisibility IsVisible = (static_cast<EMovieGraphLightModifierIntensityMethod>(IntensityMethod) == ShowWithMethod)
+				? EVisibility::Visible
+				: EVisibility::Collapsed;
+
+			return IsVisible;
+		}));
+	}
+}
+
+void FMovieGraphLightModifierCustomization::OnIntensityUnitsValuePreChange(FName IntensityUnitsPropertyName)
+{
+	if (const TSharedPtr<IDetailLayoutBuilder> DetailBuilderPin = DetailBuilder.Pin())
+	{
+		const TSharedRef<IPropertyHandle> TargetIntensityUnitsHandle = DetailBuilderPin->GetProperty(IntensityUnitsPropertyName, UMovieGraphLightModifierNode::StaticClass());
+
+		uint8 IntensityUnits = 0;
+		TargetIntensityUnitsHandle->GetValue(IntensityUnits);
+		PreviousIntensityUnits = static_cast<ELightUnits>(IntensityUnits);
+	}
+}
+
+void FMovieGraphLightModifierCustomization::OnIntensityUnitsValueChange(const FName IntensityPropertyName, const FName IntensityUnitsPropertyName) const
+{
+	if (const TSharedPtr<IDetailLayoutBuilder> DetailBuilderPin = DetailBuilder.Pin())
+	{
+		const TSharedRef<IPropertyHandle> IntensityProperty = DetailBuilderPin->GetProperty(IntensityPropertyName, UMovieGraphLightModifierNode::StaticClass());
+		const TSharedRef<IPropertyHandle> IntensityUnitsProperty = DetailBuilderPin->GetProperty(IntensityUnitsPropertyName, UMovieGraphLightModifierNode::StaticClass());
+
+		// Get the current intensity unit
+		uint8 IntensityUnits = (uint8)ELightUnits::Candelas;
+		IntensityUnitsProperty->GetValue(IntensityUnits);
+		const ELightUnits IntensityUnitsEnum = static_cast<ELightUnits>(IntensityUnits);
+
+		// Convert the intensity to a different unit if required
+		const float ConversionFactor = ULocalLightComponent::GetUnitsConversionFactor(PreviousIntensityUnits, IntensityUnitsEnum);
+
+		// Set the value of the intensity (only required if the intensity unit changes)
+		float CurrentIntensity = 0.f;
+		IntensityProperty->GetValue(CurrentIntensity);
+		IntensityProperty->SetValue(CurrentIntensity * ConversionFactor);
+
+		// The force refresh is needed when the unit changes (since it's specified via metadata)
+		DetailBuilderPin->ForceRefreshDetails();
+	}
+}
+
+TMap<UClass*, TMap<FString, TArray<const FProperty*>>> FMovieGraphLightModifierCustomization::GetCustomLightingProperties() const
+{
+	// Only properties in certain lighting categories should be included.
+	static const TArray<FString> AllowedLightingCategories = {
+		TEXT("Light"), TEXT("Performance"), TEXT("LightFunction"), TEXT("LightProfiles"), TEXT("DistanceFieldShadows"), TEXT("RayTracing")
+	};
+
+	// Some properties are excluded from being returned by this method because they're always displayed outside of the Custom menu.
+	static const TArray<FName> ExcludedClassProperties = {
+		TEXT("LightColor"), TEXT("bAffectsWorld"), TEXT("CastShadows"), TEXT("IndirectLightingIntensity"), TEXT("VolumetricScatteringIntensity"),
+		TEXT("LightingChannels"), TEXT("SamplesPerPixel"), TEXT("Intensity"), TEXT("IntensityUnits")
+	};
+
+	TMap<UClass*, TMap<FString, TArray<const FProperty*>>> CustomLightingProperties;
+
+	// Get all light classes that are available. Add the Sky Light manually since it doesn't derive from ALight.
+	TArray<UClass*> LightClasses;
+	constexpr bool bRecursive = false;
+	GetDerivedClasses(ALight::StaticClass(), LightClasses, bRecursive);
+	LightClasses.Add(ASkyLight::StaticClass());
+
+	// Discover all relevant lighting properties in the light classes.
+	for (UClass* LightClass : LightClasses)
+	{
+		TMap<FString, TArray<const FProperty*>>& LightClassProperties = CustomLightingProperties.FindOrAdd(LightClass);
+
+		// The light class itself doesn't contain the properties we care about, they're in a component (object property) in the class. Iterate the
+		// class until we find this component property.
+		const FObjectPropertyBase* LightComponentProperty = nullptr;
+		for (TFieldIterator<FProperty> PropertyIt(LightClass, EFieldIteratorFlags::IncludeSuper, EFieldIteratorFlags::ExcludeDeprecated); PropertyIt; ++PropertyIt)
+		{
+			const FProperty* LightProperty = *PropertyIt;
+
+			const FObjectPropertyBase* ObjectPropertyBase = CastField<FObjectPropertyBase>(LightProperty);
+			if (!ObjectPropertyBase)
+			{
+				continue;
+			}
+
+			if (!ObjectPropertyBase->PropertyClass->IsChildOf(ULightComponentBase::StaticClass()))
+			{
+				continue;
+			}
+
+			LightComponentProperty = ObjectPropertyBase;
+			break;
+		}
+
+		if (!LightComponentProperty)
+		{
+			continue;
+		}
+
+		// Iterate the properties of the light component.
+		for (TFieldIterator<FProperty> ComponentPropertyIt(LightComponentProperty->PropertyClass, EFieldIteratorFlags::IncludeSuper, EFieldIteratorFlags::ExcludeDeprecated); ComponentPropertyIt; ++ComponentPropertyIt)
+		{
+			const FProperty* ComponentProperty = *ComponentPropertyIt;
+			const FName PropertyName = ComponentProperty->GetFName();
+			
+			// Some properties are excluded because they're always shown outside of the Custom properties.
+			if (ExcludedClassProperties.Contains(PropertyName))
+			{
+				continue;
+			}
+			
+			// Only properties showing up in certain categories are included.
+			const FString& Category = ComponentProperty->GetMetaData(TEXT("Category"));
+			if (!AllowedLightingCategories.Contains(Category))
+			{
+				continue;
+			}
+
+			// Only include properties that show up in the details panel.
+			if (!ComponentProperty->HasAnyPropertyFlags(CPF_Edit))
+			{
+				continue;
+			}
+
+			LightClassProperties.FindOrAdd(Category).Add(ComponentProperty);
+		}
+	}
+
+	return CustomLightingProperties;
+}
+
+TSharedRef<SWidget> FMovieGraphLightModifierCustomization::GetCustomMenuContents(const TWeakInterfacePtr<IMovieGraphModifierNodeInterface> InWeakModifierNode) const
+{
+	UMovieGraphLightModifierNode* LightModifier = Cast<UMovieGraphLightModifierNode>(InWeakModifierNode.Get());
+	if (!IsValid(LightModifier))
+	{
+		return SNullWidget::NullWidget;
+	}
+
+	constexpr bool bShouldCloseWindowAfterMenuSelection = false;
+	FMenuBuilder CustomMenu(bShouldCloseWindowAfterMenuSelection, nullptr);
+
+	const TSharedRef<SWidget> CustomMenuHelpTextWidget =
+		SNew(STextBlock)
+		.Margin(FMargin(10, 5, 10, 5))
+		.ColorAndOpacity(FSlateColor::UseSubduedForeground())
+		.Text(LOCTEXT("CustomMenuHelpText", "Add an additional light actor property to this modifier."));
+
+	constexpr bool bNoIndent = true;
+	CustomMenu.AddWidget(CustomMenuHelpTextWidget, FText::GetEmpty(), bNoIndent);
+	
+	const TMap<UClass*, TMap<FString, TArray<const FProperty*>>> PropertiesToDisplay = GetCustomLightingProperties();
+
+	TArray<UClass*> LightClasses;
+	PropertiesToDisplay.GetKeys(LightClasses);
+
+	CustomMenu.BeginSection(NAME_None, LOCTEXT("LightActorTypesSectionHeader", "Light Actor Types"));
+
+	for (UClass* LightClass : LightClasses)
+	{
+		const TMap<FString, TArray<const FProperty*>>* LightClassProperties = PropertiesToDisplay.Find(LightClass);
+		if (!LightClassProperties)
+		{
+			continue;
+		}
+
+		CustomMenu.AddSubMenu(
+			LightClass->GetDisplayNameText(),
+			TAttribute<FText>(),
+			FNewMenuDelegate::CreateLambda([LightClass, InWeakModifierNode, ClassProperties = *LightClassProperties](FMenuBuilder& MenuBuilder)
+			{
+				for (const TPair<FString, TArray<const FProperty*>>& CategoryPair : ClassProperties)
+				{
+					MenuBuilder.BeginSection(NAME_None, FText::FromString(FName::NameToDisplayString(CategoryPair.Key, false)));
+					
+					for (const FProperty* CategoryProperty : CategoryPair.Value)
+					{
+						MenuBuilder.AddMenuEntry(
+							CategoryProperty->GetDisplayNameText(),
+							TAttribute<FText>(),
+							FSlateIcon(),
+							FUIAction(
+								FExecuteAction::CreateLambda([InWeakModifierNode, LightClass, CategoryProperty]()
+								{
+									if (UMovieGraphLightModifierNode* LightModifierMode = Cast<UMovieGraphLightModifierNode>(InWeakModifierNode.Get()))
+									{
+										if (LightModifierMode->HasCustomLightProperty(LightClass, CategoryProperty))
+										{
+											const FScopedTransaction Transaction(LOCTEXT("RemoveCustomLightingProperty", "Remove Custom Lighting Property"));
+											LightModifierMode->RemoveCustomLightProperty(LightClass, CategoryProperty);
+										}
+										else
+										{
+											const FScopedTransaction Transaction(LOCTEXT("AddCustomLightingProperty", "Add Custom Lighting Property"));
+											LightModifierMode->AddCustomLightProperty(LightClass, CategoryProperty);
+										}
+									}
+								}),
+								FCanExecuteAction::CreateLambda([]()
+								{
+									return true;
+								}),
+								FIsActionChecked::CreateLambda([InWeakModifierNode, LightClass, CategoryProperty]()
+								{
+									if (const UMovieGraphLightModifierNode* LightModifierMode = Cast<UMovieGraphLightModifierNode>(InWeakModifierNode.Get()))
+									{
+										return LightModifierMode->HasCustomLightProperty(LightClass, CategoryProperty);
+									}
+
+									return false;
+								})
+							),
+							NAME_None,
+							EUserInterfaceActionType::ToggleButton);
+					}
+
+					MenuBuilder.EndSection();
+				}
+			}),
+			false,
+			FSlateIconFinder::FindIconForClass(LightClass),
+			bShouldCloseWindowAfterMenuSelection);
+	}
+
+	CustomMenu.EndSection();	// "Light Actor Types" section
+
+	return CustomMenu.MakeWidget();
+}
+
+#undef LOCTEXT_NAMESPACE
diff --git a/Engine/Plugins/MovieScene/MovieRenderPipeline/Source/MovieRenderPipelineEditor/Private/Customizations/Graph/MovieGraphLightModifierCustomization.h b/Engine/Plugins/MovieScene/MovieRenderPipeline/Source/MovieRenderPipelineEditor/Private/Customizations/Graph/MovieGraphLightModifierCustomization.h
new file mode 100644
--- /dev/null
+++ b/Engine/Plugins/MovieScene/MovieRenderPipeline/Source/MovieRenderPipelineEditor/Private/Customizations/Graph/MovieGraphLightModifierCustomization.h
@@ -0,0 +1,67 @@
+﻿// Copyright Epic Games, Inc. All Rights Reserved.
+
+#pragma once
+
+#include "Engine/Scene.h"
+#include "Graph/Nodes/MovieGraphLightModifierNode.h"
+#include "IDetailCustomization.h"
+#include "UObject/WeakInterfacePtr.h"
+
+class IMovieGraphModifierNodeInterface;
+class SMovieGraphModifierCollectionsList;
+
+/** Customize how the Light Modifier node appears in the details panel. */
+class FMovieGraphLightModifierCustomization final : public IDetailCustomization
+{
+public:
+	static TSharedRef<IDetailCustomization> MakeInstance();
+
+protected:
+	//~ Begin IDetailCustomization interface
+	virtual void CustomizeDetails(IDetailLayoutBuilder& InDetailBuilder) override;
+	virtual void CustomizeDetails(const TSharedPtr<IDetailLayoutBuilder>& InDetailBuilder) override;
+	//~ End IDetailCustomization interface
+
+private:
+	/** Gets the selected modifier node that this customization is displaying. */
+	TWeakInterfacePtr<IMovieGraphModifierNodeInterface> GetSelectedModifierNode() const;
+
+	/** Gets the contents of the "Custom" menu. */
+	TSharedRef<SWidget> GetCustomMenuContents(const TWeakInterfacePtr<IMovieGraphModifierNodeInterface> InWeakModifierNode) const;
+
+	/**
+	 * Initializes the Intensity-related property to display correctly with its associated Intensity Units property. This mostly initializes the
+	 * metadata on the Intensity-related property, but may initialize other aspects of the property as well.
+	 */
+	void InitializeIntensityProperty(const FName& InIntensityPropertyName, const FName& InIntensityUnitsPropertyName);
+
+	/**
+	 * Adds a Visibility lambda to the given property, which will conditionally show the property only if the node's Intensity Method
+	 * matches the value provided in ShowWithMethod (eg, only show the property when the Intensity Method is set to PointRectSpot).
+	 */
+	void ShowPropertyConditionally(FName IntensityPropertyName, EMovieGraphLightModifierIntensityMethod ShowWithMethod);
+
+	/** Handles value pre-change logic for the Intensity Units property. */
+	void OnIntensityUnitsValuePreChange(FName IntensityUnitsPropertyName);
+
+	/** Handles value-change logic for the Intensity Units property. */
+	void OnIntensityUnitsValueChange(FName IntensityPropertyName, FName IntensityUnitsPropertyName) const;
+
+	/**
+	 * Gets all of the properties that should be displayed in the "Custom" properties menu.
+	 *
+	 * The key to the outermost TMap is the light class. The properties are contained in the inner map. In the inner map, the key is the property
+	 * category, the value is the property itself.
+	 */
+	TMap<UClass*, TMap<FString, TArray<const FProperty*>>> GetCustomLightingProperties() const;
+
+private:
+	/** The details builder associated with the customization. */
+	TWeakPtr<IDetailLayoutBuilder> DetailBuilder;
+
+	/** Displays the collections which have been chosen. */
+	TSharedPtr<SMovieGraphModifierCollectionsList> CollectionsList;
+
+	/** The value of the Intensity Units property before it changes. */
+	ELightUnits PreviousIntensityUnits = ELightUnits::Candelas;
+};
\ No newline at end of file
diff --git a/Engine/Plugins/MovieScene/MovieRenderPipeline/Source/MovieRenderPipelineEditor/Private/Customizations/Graph/MovieGraphMemberCustomization.h b/Engine/Plugins/MovieScene/MovieRenderPipeline/Source/MovieRenderPipelineEditor/Private/Customizations/Graph/MovieGraphMemberCustomization.h
--- a/Engine/Plugins/MovieScene/MovieRenderPipeline/Source/MovieRenderPipelineEditor/Private/Customizations/Graph/MovieGraphMemberCustomization.h
+++ b/Engine/Plugins/MovieScene/MovieRenderPipeline/Source/MovieRenderPipelineEditor/Private/Customizations/Graph/MovieGraphMemberCustomization.h
@@ -138,31 +138,6 @@ protected:
 			const UMovieGraphInterfaceBase* InterfaceBase = Cast<UMovieGraphInterfaceBase>(CustomizedObject);
 			TWeakObjectPtr WeakInterfaceBase(InterfaceBase);
 			
-			// Enable/disable the value property for inputs/outputs based on whether it is specified as a branch or not
-			if (InterfaceBase)
-			{
-				const TSharedRef<IPropertyHandle> ValueProperty = DetailBuilder.GetProperty("Value", UMovieGraphValueContainer::StaticClass());
-				if (ValueProperty->IsValidHandle())
-				{
-					const TAttribute<bool> IsEnabledAttribute = TAttribute<bool>::CreateLambda([WeakInterfaceBase]()
-					{
-						return WeakInterfaceBase.IsValid() ? !WeakInterfaceBase->bIsBranch : false;
-					});
-					
-					DetailBuilder.EditDefaultProperty(ValueProperty)->IsEnabled(IsEnabledAttribute);
-				}
-			}
-
-			// Enable/disable the value property for variables based on the editable state
-			if (const UMovieGraphVariable* Variable = Cast<UMovieGraphVariable>(CustomizedObject))
-			{
-				const TSharedRef<IPropertyHandle> ValueProperty = DetailBuilder.GetProperty("Value", UMovieGraphValueContainer::StaticClass());
-				if (ValueProperty->IsValidHandle())
-				{
-					DetailBuilder.EditDefaultProperty(ValueProperty)->IsEnabled(Variable->IsEditable());
-				}
-			}
-
 			UMovieGraphMember* MemberObject = Cast<UMovieGraphMember>(CustomizedObject);
 			if (!MemberObject)
 			{
@@ -253,25 +228,60 @@ protected:
 			]
 			.ValueContent()
 			[
-				SNew(SPinTypeSelector, FGetPinTypeTree::CreateLambda(GetFilteredVariableTypeTree))
-				.TargetPinType_Lambda([WeakMemberObject]()
+				SNew(SBox)
+				.Padding(FMargin(6.f, 0.f))
+				.HAlign(HAlign_Left)
+				.VAlign(VAlign_Center)
+				[
+					SNew(SPinTypeSelector, FGetPinTypeTree::CreateLambda(GetFilteredVariableTypeTree))
+					.TargetPinType_Lambda([WeakMemberObject]()
+					{
+						// The SPinTypeSelector popup might outlive this details view, so the member could be invalid
+						if (!WeakMemberObject.IsValid())
+						{
+							return FEdGraphPinType();
+						}
+
+						constexpr bool bIsBranch = false;
+						constexpr bool bIsWildcard = false;
+						return UMoviePipelineEdGraphNodeBase::GetPinType(WeakMemberObject->GetValueType(), bIsBranch, bIsWildcard, WeakMemberObject->GetValueTypeObject());
+					})
+					.OnPinTypeChanged_Lambda(PinInfoChanged, WeakMemberObject)
+					.Schema(GetDefault<UPropertyBagSchema>())
+					.bAllowArrays(false)
+					.TypeTreeFilter(ETypeTreeFilter::None)
+					.Font(IDetailLayoutBuilder::GetDetailFont())
+				]
+			];
+
+			// Display the property bag Value property. Needs to be added manually so it doesn't display in a multi-nested hierarchy.
+			const TSharedRef<IPropertyHandle> ValueProperty = DetailBuilder.GetProperty("MultiValueContainer", UMovieGraphValueContainer::StaticClass());
+			if (ValueProperty->IsValidHandle())
+			{
+				// Get the underlying "Values" property
+				const TSharedPtr<IPropertyHandle> ValuesProperty = ValueProperty->GetChildHandle("Values");
+				if (ValuesProperty.IsValid() && ValuesProperty->IsValidHandle())
 				{
-					// The SPinTypeSelector popup might outlive this details view, so the member could be invalid
-					if (!WeakMemberObject.IsValid())
+					DetailBuilder.HideProperty(ValueProperty);
+					IDetailPropertyRow& ValuePropertyRow = DetailBuilder.EditCategory("Value").AddProperty(ValuesProperty);
+
+					// Enable/disable the value property for inputs/outputs based on whether it is specified as a branch or not
+					if (InterfaceBase)
 					{
-						return FEdGraphPinType();
+						ValuePropertyRow.IsEnabled(
+							TAttribute<bool>::CreateLambda([WeakInterfaceBase]()
+							{
+								return WeakInterfaceBase.IsValid() ? !WeakInterfaceBase->bIsBranch : false;
+							}));
 					}
 
-					constexpr bool bIsBranch = false;
-					constexpr bool bIsWildcard = false;
-					return UMoviePipelineEdGraphNodeBase::GetPinType(WeakMemberObject->GetValueType(), bIsBranch, bIsWildcard, WeakMemberObject->GetValueTypeObject());
-				})
-				.OnPinTypeChanged_Lambda(PinInfoChanged, WeakMemberObject)
-				.Schema(GetDefault<UPropertyBagSchema>())
-				.bAllowArrays(false)
-				.TypeTreeFilter(ETypeTreeFilter::None)
-				.Font(IDetailLayoutBuilder::GetDetailFont())
-			];
+					// Enable/disable the value property for variables based on the editable state
+					if (const UMovieGraphVariable* Variable = Cast<UMovieGraphVariable>(CustomizedObject))
+					{
+						ValuePropertyRow.IsEnabled(Variable->IsEditable());
+					}
+				}
+			}
 		}
 	}
 	//~ End IDetailCustomization interface
diff --git a/Engine/Plugins/MovieScene/MovieRenderPipeline/Source/MovieRenderPipelineEditor/Private/Customizations/Graph/MovieGraphModifiersCustomization.cpp b/Engine/Plugins/MovieScene/MovieRenderPipeline/Source/MovieRenderPipelineEditor/Private/Customizations/Graph/MovieGraphModifiersCustomization.cpp
--- a/Engine/Plugins/MovieScene/MovieRenderPipeline/Source/MovieRenderPipelineEditor/Private/Customizations/Graph/MovieGraphModifiersCustomization.cpp
+++ b/Engine/Plugins/MovieScene/MovieRenderPipeline/Source/MovieRenderPipelineEditor/Private/Customizations/Graph/MovieGraphModifiersCustomization.cpp
@@ -6,104 +6,14 @@
 #include "DetailLayoutBuilder.h"
 #include "DetailWidgetRow.h"
 #include "Graph/MovieGraphConfig.h"
-#include "Graph/MovieGraphSharedWidgets.h"
 #include "Graph/MovieGraphRenderLayerSubsystem.h"
 #include "Graph/Nodes/MovieGraphCollectionNode.h"
 #include "Graph/Nodes/MovieGraphModifierNode.h"
-#include "ScopedTransaction.h"
-#include "Widgets/Images/SLayeredImage.h"
+#include "Widgets/Graph/SModifierCollectionsHeader.h"
+#include "Widgets/Graph/SMovieGraphModifierCollectionsList.h"
 
 #define LOCTEXT_NAMESPACE "MovieGraphModifiersCustomization"
 
-/** Discovers collections that are pickable from a specific graph, and presents them in a list. */
-class SMovieGraphCollectionPicker final : public SMovieGraphSimplePicker<FName>
-{
-public:
-	DECLARE_DELEGATE_OneParam(FOnCollectionPicked, FName);
-	DECLARE_DELEGATE_RetVal_OneParam(bool, FOnFilter, FName);
-
-	SLATE_BEGIN_ARGS(SMovieGraphCollectionPicker)
-		{}
-		/** The graph to begin discovering collections from. */
-		SLATE_ATTRIBUTE(UMovieGraphConfig*, Graph)
-
-		/** Called when a collection is picked in the list. */
-		SLATE_EVENT(FOnCollectionPicked, OnCollectionPicked);
-
-		/** Optional filter that can prevent discovered collections from showing up in the list. */
-		SLATE_EVENT(FOnFilter, OnFilter);
-	SLATE_END_ARGS()
-
-	void Construct(const FArguments& InArgs)
-	{
-		CurrentGraph = InArgs._Graph.Get();
-		OnCollectionPicked = InArgs._OnCollectionPicked;
-		OnFilter = InArgs._OnFilter;
-
-		SMovieGraphSimplePicker::Construct(SMovieGraphSimplePicker::FArguments()
-			.OnGetRowText_Lambda([](FName InCollectionName) { return FText::FromName(InCollectionName); })
-			.Title(LOCTEXT("PickCollectionHelpText", "Pick a Collection"))
-			.DataSourceEmptyMessage(LOCTEXT("NoCollectionsFoundWarning", "No collections found."))
-			.OnItemPicked(OnCollectionPicked));
-
-		// Update the data source *after* calling Construct() above because Construct() will populate DataSource based on the widget arguments, but
-		// we want to control/update it manually.
-		UpdateDataSource();
-	}
-
-private:
-	/** Discovers all available collections, and refreshes the data the list is displaying. */
-	TArray<UMovieGraphCollectionNode*> UpdateDataSource()
-	{
-		TArray<UMovieGraphCollectionNode*> CollectionNodes;
-		
-		if (!CurrentGraph)
-		{
-			return CollectionNodes;
-		}
-
-		TSet<UMovieGraphConfig*> Subgraphs;
-		CurrentGraph->GetAllContainedSubgraphs(Subgraphs);
-
-		TArray<UMovieGraphConfig*> AllGraphs = { CurrentGraph };
-		for (UMovieGraphConfig* Subgraph : Subgraphs)
-		{
-			AllGraphs.Add(Subgraph);
-		}
-
-		for (const UMovieGraphConfig* Graph : AllGraphs)
-		{
-			for (const TObjectPtr<UMovieGraphNode>& Node : Graph->GetNodes())
-			{
-				if (const TObjectPtr<UMovieGraphCollectionNode> CollectionNode = Cast<UMovieGraphCollectionNode>(Node))
-				{
-					const FName CollectionName = FName(CollectionNode->Collection->GetCollectionName());
-					bool bIncludeCollection = true;
-
-					if (OnFilter.IsBound())
-					{
-						bIncludeCollection = OnFilter.Execute(CollectionName);
-					}
-
-					if (bIncludeCollection)
-					{
-						DataSource.AddUnique(CollectionName);
-					}
-				}
-			}
-		}
-
-		return CollectionNodes;
-	}
-
-private:
-	/** The current graph being viewed. The data source will be populated from this graph. */
-	UMovieGraphConfig* CurrentGraph = nullptr;
-	
-	FOnCollectionPicked OnCollectionPicked;
-	FOnFilter OnFilter;
-};
-
 TSharedRef<IDetailCustomization> FMovieGraphModifiersCustomization::MakeInstance()
 {
 	return MakeShared<FMovieGraphModifiersCustomization>();
@@ -112,124 +22,30 @@ TSharedRef<IDetailCustomization> FMovieGraphModifiersCustomization::MakeInstance
 void FMovieGraphModifiersCustomization::CustomizeDetails(IDetailLayoutBuilder& InDetailBuilder)
 {
 	// The customization only supports editing a single Modifier node
-	const TWeakObjectPtr<UMovieGraphModifierNode> ModifierNode = GetSelectedModifierNode();
+	const TWeakInterfacePtr<IMovieGraphModifierNodeInterface> ModifierNode = GetSelectedModifierNode();
 	if (!ModifierNode.IsValid())
 	{
 		return;
 	}
 
-	// Update the data source
-	RefreshListDataSource();
-	
-	// Generate a (multi-layered) icon for the "Add" menu
-	const TSharedRef<SLayeredImage> AddIcon =
-		SNew(SLayeredImage)
-		.ColorAndOpacity(FSlateColor::UseForeground())
-		.Image(FAppStyle::GetBrush("LevelEditor.OpenAddContent.Background"));
-	AddIcon->AddLayer(FAppStyle::GetBrush("LevelEditor.OpenAddContent.Overlay"));
-
 	// Replace the "Collection" category row with a custom whole-row widget which includes an add-collection button
 	IDetailCategoryBuilder& CollectionCategory = InDetailBuilder.EditCategory(FName("Collection"), FText::GetEmpty(), ECategoryPriority::Uncommon);
 	CollectionCategory.HeaderContent
 	(
-		SNew(SHorizontalBox)
-		+ SHorizontalBox::Slot()
-		.HAlign(HAlign_Left)
-		.VAlign(VAlign_Center)
-		.AutoWidth()
-		[
-			SNew(STextBlock)
-			.Text(FText::FromString("Collections"))
-			.Font(FAppStyle::Get().GetFontStyle("DetailsView.CategoryFontStyle"))
-		]
-
-		+ SHorizontalBox::Slot()
-		.HAlign(HAlign_Left)
-		.VAlign(VAlign_Center)
-		.Padding(5.f, 0, 0, 0)
-		[
-			SNew(SComboButton)
-			.ToolTipText(LOCTEXT("AddCollectionToModifierTooltip", "Add a collection that will be affected by the configured modifiers."))
-			.ComboButtonStyle(FAppStyle::Get(), "SimpleComboButton")
-			.ContentPadding(0)
-			.HasDownArrow(false)
-			.OnGetMenuContent_Lambda([ModifierNode, this]()
-			{
-				return
-					SNew(SMovieGraphCollectionPicker)
-					.Graph(ModifierNode->GetTypedOuter<UMovieGraphConfig>())
-					.OnFilter_Lambda([ModifierNode](const FName CollectionName)
-					{
-						if (ModifierNode.IsValid())
-						{
-							return !ModifierNode.Get()->GetAllCollections().Contains(CollectionName);
-						}
-
-						return false;
-					})
-					.OnCollectionPicked_Lambda([ModifierNode, this](const FName PickedCollectionName)
-					{
-						if (ModifierNode.IsValid())
-						{
-							const FScopedTransaction Transaction(LOCTEXT("AddCollectionToModifier", "Add Collection to Modifier"));
-							
-							ModifierNode->AddCollection(PickedCollectionName);
-							ListDataSource = ModifierNode->GetAllCollections();
-							CollectionsList->Refresh();
-						}
-					});
-			})
-			.ButtonContent()
-			[
-				SNew(SBox)
-				.WidthOverride(16)
-				.HeightOverride(16)
-				[
-					AddIcon
-				]
-			]
-		]
+		SNew(SMovieGraphCollectionsHeaderWidget)
+		.WeakModifierInterface(ModifierNode)
+		.OnCollectionPicked_Lambda([this](const FName CollectionName)
+		{
+			CollectionsList->Refresh();
+		})
 	, /* bWholeRowContent */ true);
 
 	// Add a collections browser
 	CollectionCategory.AddCustomRow(FText::GetEmpty())
 	.WholeRowWidget
 	[
-		SAssignNew(CollectionsList, SMovieGraphSimpleList<FName>)
-		.DataSource(&ListDataSource)
-		.DataType(FText::FromString("Collection"))
-		.DataTypePlural(FText::FromString("Collections"))
-		.OnDelete_Lambda([this, ModifierNode](const TArray<FName> DeletedCollectionNames)
-		{
-			if (ModifierNode.IsValid())
-			{
-				const FScopedTransaction Transaction(LOCTEXT("RemoveCollectionsFromModifier", "Remove Collections from Modifier"));
-
-				for (const FName& DeletedCollectionName : DeletedCollectionNames)
-				{
-					ModifierNode.Get()->RemoveCollection(DeletedCollectionName);
-				}
-				
-				ListDataSource = ModifierNode->GetAllCollections();
-				CollectionsList->Refresh();
-			}
-		})
-		.OnGetRowText_Static(&GetCollectionRowText)
-		.ShowEnableDisable(true)
-		.OnGetRowEnableState_Lambda([ModifierNode](FName InCollectionName)
-		{
-			return ModifierNode.IsValid() ? ModifierNode.Get()->IsCollectionEnabled(InCollectionName) : true;
-		})
-		.OnSetRowEnableState_Lambda([ModifierNode](FName InCollectionName, bool bNewEnableState)
-		{
-			if (ModifierNode.IsValid())
-			{
-				const FScopedTransaction Transaction(LOCTEXT("ChangeCollectionEnableState", "Change Collection Enable State"));
-				
-				ModifierNode.Get()->SetCollectionEnabled(InCollectionName, bNewEnableState);
-			}
-		})
-		.OnRefreshDataSourceRequested(this, &FMovieGraphModifiersCustomization::RefreshListDataSource)
+		SAssignNew(CollectionsList, SMovieGraphModifierCollectionsList)
+		.WeakModifierInterface(ModifierNode)
 	];
 
 	// For all modifiers added to the node, add a category for each, and add each modifier's EditAnywhere properties to the category
@@ -264,17 +80,7 @@ void FMovieGraphModifiersCustomization::CustomizeDetails(const TSharedPtr<IDetai
 	CustomizeDetails(*InDetailBuilder);
 }
 
-const FSlateBrush* FMovieGraphModifiersCustomization::GetCollectionRowIcon(const FName CollectionName)
-{
-	return FAppStyle::GetBrush("Icons.FilledCircle");
-}
-
-FText FMovieGraphModifiersCustomization::GetCollectionRowText(const FName CollectionName)
-{
-	return FText::FromName(CollectionName);
-}
-
-TWeakObjectPtr<UMovieGraphModifierNode> FMovieGraphModifiersCustomization::GetSelectedModifierNode() const
+TWeakInterfacePtr<IMovieGraphModifierNodeInterface> FMovieGraphModifiersCustomization::GetSelectedModifierNode() const
 {
 	if (const TSharedPtr<IDetailLayoutBuilder> DetailBuilderPin = DetailBuilder.Pin())
 	{
@@ -286,24 +92,13 @@ TWeakObjectPtr<UMovieGraphModifierNode> FMovieGraphModifiersCustomization::GetSe
 		}
 
 		const TWeakObjectPtr<UMovieGraphModifierNode> ModifierNode = ModifierNodes[0];
-		if (ModifierNode.IsValid())
+		if (ModifierNode.IsValid() && ModifierNode->Implements<UMovieGraphModifierNodeInterface>())
 		{
-			return ModifierNode;
+			return TWeakInterfacePtr<IMovieGraphModifierNodeInterface>(ModifierNode.Get());
 		}
 	}
 
 	return nullptr;
 }
 
-void FMovieGraphModifiersCustomization::RefreshListDataSource()
-{
-	const TWeakObjectPtr<UMovieGraphModifierNode> ModifierNode = GetSelectedModifierNode();
-	
-	if (const TStrongObjectPtr<UMovieGraphModifierNode> ModifierNodePin = ModifierNode.Pin())
-	{
-		// Update the data source
-		ListDataSource = ModifierNodePin->GetAllCollections();
-	}
-}
-
 #undef LOCTEXT_NAMESPACE
\ No newline at end of file
diff --git a/Engine/Plugins/MovieScene/MovieRenderPipeline/Source/MovieRenderPipelineEditor/Private/Customizations/Graph/MovieGraphModifiersCustomization.h b/Engine/Plugins/MovieScene/MovieRenderPipeline/Source/MovieRenderPipelineEditor/Private/Customizations/Graph/MovieGraphModifiersCustomization.h
--- a/Engine/Plugins/MovieScene/MovieRenderPipeline/Source/MovieRenderPipelineEditor/Private/Customizations/Graph/MovieGraphModifiersCustomization.h
+++ b/Engine/Plugins/MovieScene/MovieRenderPipeline/Source/MovieRenderPipelineEditor/Private/Customizations/Graph/MovieGraphModifiersCustomization.h
@@ -3,10 +3,10 @@
 #pragma once
 
 #include "IDetailCustomization.h"
-#include "Graph/MovieGraphSharedWidgets.h"
+#include "UObject/WeakInterfacePtr.h"
 
-class UMovieGraphCollectionNode;
-class UMovieGraphModifierNode;
+class IMovieGraphModifierNodeInterface;
+class SMovieGraphModifierCollectionsList;
 
 /** Customize how the Modifier node appears in the details panel. */
 class FMovieGraphModifiersCustomization final : public IDetailCustomization
@@ -21,22 +21,13 @@ protected:
 	//~ End IDetailCustomization interface
 
 private:
-	static const FSlateBrush* GetCollectionRowIcon(const FName CollectionName);
-	static FText GetCollectionRowText(const FName CollectionName);
-
 	/** Gets the selected modifier node that this customization is displaying. */
-	TWeakObjectPtr<UMovieGraphModifierNode> GetSelectedModifierNode() const;
-	
-	/** Refreshes the list's data source to reflect the data model. */
-	void RefreshListDataSource();
+	TWeakInterfacePtr<IMovieGraphModifierNodeInterface> GetSelectedModifierNode() const;
 
 private:
 	/** The details builder associated with the customization. */
 	TWeakPtr<IDetailLayoutBuilder> DetailBuilder;
 
-	/** The data source for the list view. Modifiers don't expose a direct pointer to the array of collections, which is required by the list view. */
-	TArray<FName> ListDataSource;
-
 	/** Displays the collections which have been chosen. */
-	TSharedPtr<SMovieGraphSimpleList<FName>> CollectionsList;
+	TSharedPtr<SMovieGraphModifierCollectionsList> CollectionsList;
 };
\ No newline at end of file
diff --git a/Engine/Plugins/MovieScene/MovieRenderPipeline/Source/MovieRenderPipelineEditor/Private/Customizations/Graph/MovieGraphNodeCustomization.h b/Engine/Plugins/MovieScene/MovieRenderPipeline/Source/MovieRenderPipelineEditor/Private/Customizations/Graph/MovieGraphNodeCustomization.h
--- a/Engine/Plugins/MovieScene/MovieRenderPipeline/Source/MovieRenderPipelineEditor/Private/Customizations/Graph/MovieGraphNodeCustomization.h
+++ b/Engine/Plugins/MovieScene/MovieRenderPipeline/Source/MovieRenderPipelineEditor/Private/Customizations/Graph/MovieGraphNodeCustomization.h
@@ -87,6 +87,13 @@ private:
 		DetailBuilder.GetCategoryNames(CategoryNames);
 		for (const FName& CategoryName : CategoryNames)
 		{
+			// Only go through top-level categories. Subcategories are currently not supported here (EditCategory() does not
+			// work properly with category names like "Foo|Bar").
+			if (CategoryName.ToString().Contains(TEXT("|")))
+			{
+				continue;
+			}
+
 			IDetailCategoryBuilder& CategoryBuilder = DetailBuilder.EditCategory(CategoryName);
 
 			TArray<TSharedRef<IPropertyHandle>> CategoryProperties;
diff --git a/Engine/Plugins/MovieScene/MovieRenderPipeline/Source/MovieRenderPipelineEditor/Private/Graph/MovieGraphAssetToolkit.cpp b/Engine/Plugins/MovieScene/MovieRenderPipeline/Source/MovieRenderPipelineEditor/Private/Graph/MovieGraphAssetToolkit.cpp
--- a/Engine/Plugins/MovieScene/MovieRenderPipeline/Source/MovieRenderPipelineEditor/Private/Graph/MovieGraphAssetToolkit.cpp
+++ b/Engine/Plugins/MovieScene/MovieRenderPipeline/Source/MovieRenderPipelineEditor/Private/Graph/MovieGraphAssetToolkit.cpp
@@ -5,6 +5,7 @@
 #include "Customizations/Graph/MovieGraphApplyCVarPresetNodeCustomization.h"
 #include "Customizations/Graph/MovieGraphCollectionsCustomization.h"
 #include "Customizations/Graph/MovieGraphExecuteScriptNodeCustomization.h"
+#include "Customizations/Graph/MovieGraphLightModifierCustomization.h"
 #include "Customizations/Graph/MovieGraphMultiLayerEXRNodeCustomization.h"
 #include "Customizations/Graph/MovieGraphMaterialParameterCollectionModifierNodeCustomization.h"
 #include "Customizations/Graph/MovieGraphMemberCustomization.h"
@@ -23,6 +24,7 @@
 #include "Graph/Nodes/MovieGraphApplyCVarPresetNode.h"
 #include "Graph/Nodes/MovieGraphCollectionNode.h"
 #include "Graph/Nodes/MovieGraphExecuteScriptNode.h"
+#include "Graph/Nodes/MovieGraphLightModifierNode.h"
 #include "Graph/Nodes/MovieGraphMaterialParameterCollectionModifierNode.h"
 #include "Graph/Nodes/MovieGraphModifierNode.h"
 #include "Graph/Nodes/MovieGraphPathTracerPassNode.h"
@@ -436,6 +438,10 @@ TSharedRef<SDockTab> FMovieGraphAssetToolkit::SpawnTab_RenderGraphDetails(const
 		UMovieGraphModifierNode::StaticClass(),
 		FOnGetDetailCustomizationInstance::CreateStatic(&FMovieGraphModifiersCustomization::MakeInstance));
 
+	SelectedGraphObjectsDetailsWidget->RegisterInstancedCustomPropertyLayout(
+		UMovieGraphLightModifierNode::StaticClass(),
+		FOnGetDetailCustomizationInstance::CreateStatic(&FMovieGraphLightModifierCustomization::MakeInstance));
+
 	SelectedGraphObjectsDetailsWidget->RegisterInstancedCustomPropertyTypeLayout(
 		FMoviePipelinePostProcessPass::StaticStruct()->GetFName(),
 		FOnGetPropertyTypeCustomizationInstance::CreateStatic(&FMovieGraphPostProcessPassCustomization::MakeInstance));
@@ -505,10 +511,16 @@ void FMovieGraphAssetToolkit::RegisterDetailsViewSections()
 		for (TFieldIterator<FProperty> PropertyIt(NodeClass); PropertyIt; ++PropertyIt)
 		{
 			const FProperty* NodeProperty = *PropertyIt;
-
 			const FString* CategoryMetadata = NodeProperty->FindMetaData(CategoryMetadataKey);
+
 			if (CategoryMetadata && !CategoryMetadata->IsEmpty() && !CategoriesToExclude.Contains(*CategoryMetadata))
 			{
+				// Only top-level categories can be added as sections.
+				if (CategoryMetadata->Contains(TEXT("|")))
+				{
+					continue;
+				}
+
 				const FName SectionName = FName(*CategoryMetadata);
 				const FText LocalizedCategory = FObjectEditorUtils::GetCategoryText(NodeProperty);
 				
diff --git a/Engine/Plugins/MovieScene/MovieRenderPipeline/Source/MovieRenderPipelineEditor/Private/Widgets/Graph/SModifierCollectionsHeader.cpp b/Engine/Plugins/MovieScene/MovieRenderPipeline/Source/MovieRenderPipelineEditor/Private/Widgets/Graph/SModifierCollectionsHeader.cpp
new file mode 100644
--- /dev/null
+++ b/Engine/Plugins/MovieScene/MovieRenderPipeline/Source/MovieRenderPipelineEditor/Private/Widgets/Graph/SModifierCollectionsHeader.cpp
@@ -0,0 +1,93 @@
+﻿// Copyright Epic Games, Inc. All Rights Reserved.
+
+#include "Widgets/Graph/SModifierCollectionsHeader.h"
+
+#include "ScopedTransaction.h"
+#include "Graph/Nodes/MovieGraphModifierNode.h"
+#include "Widgets/Images/SLayeredImage.h"
+#include "Widgets/Layout/SBox.h"
+#include "Widgets/Text/STextBlock.h"
+#include "Widgets/Graph/SMovieGraphCollectionPicker.h"
+
+#define LOCTEXT_NAMESPACE "MovieGraphModifierCollectionsHeader"
+
+void SMovieGraphCollectionsHeaderWidget::Construct(const FArguments& InArgs)
+{
+	WeakModifierInterface = InArgs._WeakModifierInterface.Get();
+	OnCollectionPicked = InArgs._OnCollectionPicked;
+
+	// Generate a (multi-layered) icon for the "Add" menu
+	const TSharedRef<SLayeredImage> AddIcon =
+		SNew(SLayeredImage)
+		.ColorAndOpacity(FSlateColor::UseForeground())
+		.Image(FAppStyle::GetBrush("LevelEditor.OpenAddContent.Background"));
+	AddIcon->AddLayer(FAppStyle::GetBrush("LevelEditor.OpenAddContent.Overlay"));
+
+	ChildSlot
+	[
+		SNew(SHorizontalBox)
+		+ SHorizontalBox::Slot()
+		.HAlign(HAlign_Left)
+		.VAlign(VAlign_Center)
+		.AutoWidth()
+		[
+			SNew(STextBlock)
+			.Text(LOCTEXT("CollectionsHeaderText", "Collections"))
+			.Font(FAppStyle::Get().GetFontStyle("DetailsView.CategoryFontStyle"))
+		]
+
+		+ SHorizontalBox::Slot()
+		.HAlign(HAlign_Left)
+		.VAlign(VAlign_Center)
+		.Padding(5.f, 0, 0, 0)
+		[
+			SNew(SComboButton)
+			.ToolTipText(LOCTEXT("AddCollectionToModifierTooltip", "Add a collection that will be affected by the configured modifiers."))
+			.ComboButtonStyle(FAppStyle::Get(), "SimpleComboButton")
+			.ContentPadding(0)
+			.HasDownArrow(false)
+			.OnGetMenuContent_Lambda([this]()
+			{
+				const UObject* ModifierNode = WeakModifierInterface.GetObject();
+				
+				return
+					SNew(SMovieGraphCollectionPicker)
+					.Graph(IsValid(ModifierNode) ? ModifierNode->GetTypedOuter<UMovieGraphConfig>() : nullptr)
+					.OnFilter_Lambda([this](const FName CollectionName)
+					{
+						if (WeakModifierInterface.IsValid())
+						{
+							return !WeakModifierInterface.Get()->GetAllCollections().Contains(CollectionName);
+						}
+
+						return false;
+					})
+					.OnCollectionPicked_Lambda([this](const FName PickedCollectionName)
+					{
+						if (WeakModifierInterface.IsValid())
+						{
+							const FScopedTransaction Transaction(LOCTEXT("AddCollectionToModifier", "Add Collection to Modifier"));
+							
+							WeakModifierInterface.Get()->AddCollection(PickedCollectionName);
+
+							if (OnCollectionPicked.IsBound())
+							{
+								OnCollectionPicked.Execute(PickedCollectionName);
+							}
+						}
+					});
+			})
+			.ButtonContent()
+			[
+				SNew(SBox)
+				.WidthOverride(16)
+				.HeightOverride(16)
+				[
+					AddIcon
+				]
+			]
+		]
+	];
+}
+
+#undef LOCTEXT_NAMESPACE
diff --git a/Engine/Plugins/MovieScene/MovieRenderPipeline/Source/MovieRenderPipelineEditor/Private/Widgets/Graph/SMovieGraphCollectionPicker.cpp b/Engine/Plugins/MovieScene/MovieRenderPipeline/Source/MovieRenderPipelineEditor/Private/Widgets/Graph/SMovieGraphCollectionPicker.cpp
new file mode 100644
--- /dev/null
+++ b/Engine/Plugins/MovieScene/MovieRenderPipeline/Source/MovieRenderPipelineEditor/Private/Widgets/Graph/SMovieGraphCollectionPicker.cpp
@@ -0,0 +1,66 @@
+﻿// Copyright Epic Games, Inc. All Rights Reserved.
+
+#include "Widgets/Graph/SMovieGraphCollectionPicker.h"
+
+#include "Graph/MovieGraphConfig.h"
+#include "Graph/Nodes/MovieGraphCollectionNode.h"
+
+#define LOCTEXT_NAMESPACE "MovieGraphCollectionPicker"
+
+void SMovieGraphCollectionPicker::Construct(const FArguments& InArgs)
+{
+	CurrentGraph = InArgs._Graph.Get();
+	OnCollectionPicked = InArgs._OnCollectionPicked;
+	OnFilter = InArgs._OnFilter;
+
+	SMovieGraphSimplePicker::Construct(SMovieGraphSimplePicker::FArguments()
+		.OnGetRowText_Lambda([](FName InCollectionName) { return FText::FromName(InCollectionName); })
+		.Title(LOCTEXT("PickCollectionHelpText", "Pick a Collection"))
+		.DataSourceEmptyMessage(LOCTEXT("NoCollectionsFoundWarning", "No collections found."))
+		.OnItemPicked(OnCollectionPicked));
+
+	// Update the data source *after* calling Construct() above because Construct() will populate DataSource based on the widget arguments, but
+	// we want to control/update it manually.
+	UpdateDataSource();
+}
+
+void SMovieGraphCollectionPicker::UpdateDataSource()
+{
+	if (!CurrentGraph)
+	{
+		return;
+	}
+
+	TSet<UMovieGraphConfig*> Subgraphs;
+	CurrentGraph->GetAllContainedSubgraphs(Subgraphs);
+
+	TArray<UMovieGraphConfig*> AllGraphs = { CurrentGraph };
+	for (UMovieGraphConfig* Subgraph : Subgraphs)
+	{
+		AllGraphs.Add(Subgraph);
+	}
+
+	for (const UMovieGraphConfig* Graph : AllGraphs)
+	{
+		for (const TObjectPtr<UMovieGraphNode>& Node : Graph->GetNodes())
+		{
+			if (const TObjectPtr<UMovieGraphCollectionNode> CollectionNode = Cast<UMovieGraphCollectionNode>(Node))
+			{
+				const FName CollectionName = FName(CollectionNode->Collection->GetCollectionName());
+				bool bIncludeCollection = true;
+
+				if (OnFilter.IsBound())
+				{
+					bIncludeCollection = OnFilter.Execute(CollectionName);
+				}
+
+				if (bIncludeCollection)
+				{
+					DataSource.AddUnique(CollectionName);
+				}
+			}
+		}
+	}
+}
+
+#undef LOCTEXT_NAMESPACE
\ No newline at end of file
diff --git a/Engine/Plugins/MovieScene/MovieRenderPipeline/Source/MovieRenderPipelineEditor/Private/Widgets/Graph/SMovieGraphModifierCollectionsList.cpp b/Engine/Plugins/MovieScene/MovieRenderPipeline/Source/MovieRenderPipelineEditor/Private/Widgets/Graph/SMovieGraphModifierCollectionsList.cpp
new file mode 100644
--- /dev/null
+++ b/Engine/Plugins/MovieScene/MovieRenderPipeline/Source/MovieRenderPipelineEditor/Private/Widgets/Graph/SMovieGraphModifierCollectionsList.cpp
@@ -0,0 +1,84 @@
+﻿// Copyright Epic Games, Inc. All Rights Reserved.
+
+#include "Widgets/Graph/SMovieGraphModifierCollectionsList.h"
+
+#include "Graph/Nodes/MovieGraphModifierNode.h"
+#include "ScopedTransaction.h"
+
+#define LOCTEXT_NAMESPACE "MovieGraphModifierCollectionsList"
+
+void SMovieGraphModifierCollectionsList::Construct(const FArguments& InArgs)
+{
+	WeakModifierInterface = InArgs._WeakModifierInterface.Get();
+	
+	ChildSlot
+	[
+		SAssignNew(CollectionsList, SMovieGraphSimpleList<FName>)
+        .DataSource(&ListDataSource)
+        .DataType(FText::FromString("Collection"))
+        .DataTypePlural(FText::FromString("Collections"))
+        .OnDelete_Lambda([this](const TArray<FName> DeletedCollectionNames)
+        {
+        	if (WeakModifierInterface.IsValid())
+        	{
+        		const FScopedTransaction Transaction(LOCTEXT("RemoveCollectionsFromModifier", "Remove Collections from Modifier"));
+
+        		for (const FName& DeletedCollectionName : DeletedCollectionNames)
+        		{
+        			WeakModifierInterface.Get()->RemoveCollection(DeletedCollectionName);
+        		}
+        		
+        		Refresh();
+        	}
+        })
+        .OnGetRowText_Static(&GetCollectionRowText)
+        .ShowEnableDisable(true)
+        .OnGetRowEnableState_Lambda([this](FName InCollectionName)
+        {
+        	if (WeakModifierInterface.IsValid())
+        	{
+        		return WeakModifierInterface.Get()->IsCollectionEnabled(InCollectionName);
+        	}
+
+        	return true;
+        })
+        .OnSetRowEnableState_Lambda([this](FName InCollectionName, bool bNewEnableState)
+        {
+        	if (WeakModifierInterface.IsValid())
+        	{
+        		const FScopedTransaction Transaction(LOCTEXT("ChangeCollectionEnableState", "Change Collection Enable State"));
+        		
+        		WeakModifierInterface.Get()->SetCollectionEnabled(InCollectionName, bNewEnableState);
+        	}
+        })
+        .OnRefreshDataSourceRequested(this, &SMovieGraphModifierCollectionsList::RefreshListDataSource)
+	];
+
+	// Make sure the data source is up-to-date initially
+	Refresh();
+}
+
+void SMovieGraphModifierCollectionsList::Refresh()
+{
+	RefreshListDataSource();
+
+	if (CollectionsList)
+	{
+		CollectionsList->Refresh();
+	}
+}
+
+FText SMovieGraphModifierCollectionsList::GetCollectionRowText(const FName CollectionName)
+{
+	return FText::FromName(CollectionName);
+}
+
+void SMovieGraphModifierCollectionsList::RefreshListDataSource()
+{
+	if (WeakModifierInterface.IsValid())
+	{
+		ListDataSource = WeakModifierInterface.Get()->GetAllCollections();
+	}
+}
+
+#undef LOCTEXT_NAMESPACE
diff --git a/Engine/Plugins/MovieScene/MovieRenderPipeline/Source/MovieRenderPipelineEditor/Public/Widgets/Graph/SModifierCollectionsHeader.h b/Engine/Plugins/MovieScene/MovieRenderPipeline/Source/MovieRenderPipelineEditor/Public/Widgets/Graph/SModifierCollectionsHeader.h
new file mode 100644
--- /dev/null
+++ b/Engine/Plugins/MovieScene/MovieRenderPipeline/Source/MovieRenderPipelineEditor/Public/Widgets/Graph/SModifierCollectionsHeader.h
@@ -0,0 +1,34 @@
+﻿// Copyright Epic Games, Inc. All Rights Reserved.
+
+#pragma once
+
+#include "UObject/WeakInterfacePtr.h"
+#include "Widgets/SCompoundWidget.h"
+
+class IMovieGraphModifierNodeInterface;
+
+/**
+ * A widget which displays the header for modifier collections in the details panel.
+ */
+class SMovieGraphCollectionsHeaderWidget final : public SCompoundWidget
+{
+public:
+	DECLARE_DELEGATE_OneParam(FOnCollectionPicked, FName);
+
+	SLATE_BEGIN_ARGS(SMovieGraphCollectionsHeaderWidget) { }
+		/** The modifier node that is associated with the collections. */
+		SLATE_ATTRIBUTE(TWeakInterfacePtr<IMovieGraphModifierNodeInterface>, WeakModifierInterface)
+
+		/** Called when a collection is picked in the list. */
+		SLATE_EVENT(FOnCollectionPicked, OnCollectionPicked);
+	SLATE_END_ARGS()
+
+	void Construct(const FArguments& InArgs);
+
+private:
+	/** The modifier node associated with this widget (if any). */
+	TWeakInterfacePtr<IMovieGraphModifierNodeInterface> WeakModifierInterface;
+
+	/** Called when a collection is picked in the list. */
+	FOnCollectionPicked OnCollectionPicked;
+};
\ No newline at end of file
diff --git a/Engine/Plugins/MovieScene/MovieRenderPipeline/Source/MovieRenderPipelineEditor/Public/Widgets/Graph/SMovieGraphCollectionPicker.h b/Engine/Plugins/MovieScene/MovieRenderPipeline/Source/MovieRenderPipelineEditor/Public/Widgets/Graph/SMovieGraphCollectionPicker.h
new file mode 100644
--- /dev/null
+++ b/Engine/Plugins/MovieScene/MovieRenderPipeline/Source/MovieRenderPipelineEditor/Public/Widgets/Graph/SMovieGraphCollectionPicker.h
@@ -0,0 +1,45 @@
+﻿// Copyright Epic Games, Inc. All Rights Reserved.
+
+#pragma once
+
+#include "Graph/MovieGraphSharedWidgets.h"
+#include "Graph/Nodes/MovieGraphModifierNode.h"
+
+class UMovieGraphCollectionNode;
+class UMovieGraphConfig;
+
+/** Discovers collections that are pickable from a specific graph, and presents them in a list. */
+class SMovieGraphCollectionPicker final : public SMovieGraphSimplePicker<FName>
+{
+public:
+	DECLARE_DELEGATE_OneParam(FOnCollectionPicked, FName);
+	DECLARE_DELEGATE_RetVal_OneParam(bool, FOnFilter, FName);
+
+	SLATE_BEGIN_ARGS(SMovieGraphCollectionPicker)
+		{}
+		/** The graph to begin discovering collections from. */
+		SLATE_ATTRIBUTE(UMovieGraphConfig*, Graph)
+
+		/** Called when a collection is picked in the list. */
+		SLATE_EVENT(FOnCollectionPicked, OnCollectionPicked);
+
+		/** Optional filter that can prevent discovered collections from showing up in the list. */
+		SLATE_EVENT(FOnFilter, OnFilter);
+	SLATE_END_ARGS()
+
+	void Construct(const FArguments& InArgs);
+
+private:
+	/** Discovers all available collections, and refreshes the data the list is displaying. */
+	void UpdateDataSource();
+
+private:
+	/** The current graph being viewed. The data source will be populated from this graph. */
+	UMovieGraphConfig* CurrentGraph = nullptr;
+
+	/** Called when a collection is picked in the list. */
+	FOnCollectionPicked OnCollectionPicked;
+
+	/** Optional filter that can prevent discovered collections from showing up in the list. */
+	FOnFilter OnFilter;
+};
\ No newline at end of file
diff --git a/Engine/Plugins/MovieScene/MovieRenderPipeline/Source/MovieRenderPipelineEditor/Public/Widgets/Graph/SMovieGraphModifierCollectionsList.h b/Engine/Plugins/MovieScene/MovieRenderPipeline/Source/MovieRenderPipelineEditor/Public/Widgets/Graph/SMovieGraphModifierCollectionsList.h
new file mode 100644
--- /dev/null
+++ b/Engine/Plugins/MovieScene/MovieRenderPipeline/Source/MovieRenderPipelineEditor/Public/Widgets/Graph/SMovieGraphModifierCollectionsList.h
@@ -0,0 +1,39 @@
+﻿// Copyright Epic Games, Inc. All Rights Reserved.
+
+#pragma once
+
+#include "Graph/MovieGraphSharedWidgets.h"
+#include "UObject/WeakInterfacePtr.h"
+
+class IMovieGraphModifierNodeInterface;
+
+/** Lists out the collections that have been chosen for a specific modifier node. */
+class SMovieGraphModifierCollectionsList final : public SCompoundWidget
+{
+public:
+	SLATE_BEGIN_ARGS(SMovieGraphModifierCollectionsList) { }
+		/** The modifier node that is associated with the collections. */
+		SLATE_ATTRIBUTE(TWeakInterfacePtr<IMovieGraphModifierNodeInterface>, WeakModifierInterface)
+	SLATE_END_ARGS()
+
+	void Construct(const FArguments& InArgs);
+
+	/** Refreshes the widget and its internal data source. */
+	void Refresh();
+
+private:
+	static FText GetCollectionRowText(const FName CollectionName);
+
+	/** Refreshes the list's data source to reflect the data model. */
+	void RefreshListDataSource();
+
+private:
+	/** The data source for the list view. Modifiers don't expose a direct pointer to the array of collections, which is required by the list view. */
+	TArray<FName> ListDataSource;
+	
+	/** The modifier node associated with this widget (if any). */
+	TWeakInterfacePtr<IMovieGraphModifierNodeInterface> WeakModifierInterface;
+
+	/** Displays the collections which have been chosen. */
+	TSharedPtr<SMovieGraphSimpleList<FName>> CollectionsList;
+};
\ No newline at end of file
