MoviePipeline: Fixed a CIS error which indicated an uninitialized struct member.

#Synced-CL 49282327

diff --git a/Engine/Plugins/MovieScene/MovieRenderPipeline/Source/MovieRenderPipelineCore/Public/Graph/MovieGraphValueView.h b/Engine/Plugins/MovieScene/MovieRenderPipeline/Source/MovieRenderPipelineCore/Public/Graph/MovieGraphValueView.h
--- a/Engine/Plugins/MovieScene/MovieRenderPipeline/Source/MovieRenderPipelineCore/Public/Graph/MovieGraphValueView.h
+++ b/Engine/Plugins/MovieScene/MovieRenderPipeline/Source/MovieRenderPipelineCore/Public/Graph/MovieGraphValueView.h
@@ -48,7 +48,7 @@ struct FMovieGraphValueViewProperty
 
 	/** The type of value the property holds. */
 	UPROPERTY(BlueprintReadWrite, Category = "Property")
-	EMovieGraphValueType ValueType;
+	EMovieGraphValueType ValueType = EMovieGraphValueType::None;
 
 	/** The value type object for object, struct, and enum-typed properties (nullptr for other types). */
 	UPROPERTY(BlueprintReadWrite, Category = "Property")
