MoviePipeline: Added a new "Light Actor Type" condition group query to the graph.

#Synced-CL 49282327

diff --git a/Engine/Plugins/MovieScene/MovieRenderPipeline/Source/MovieRenderPipelineCore/Private/Graph/MovieGraphRenderLayerSubsystem.cpp b/Engine/Plugins/MovieScene/MovieRenderPipeline/Source/MovieRenderPipelineCore/Private/Graph/MovieGraphRenderLayerSubsystem.cpp
--- a/Engine/Plugins/MovieScene/MovieRenderPipeline/Source/MovieRenderPipelineCore/Private/Graph/MovieGraphRenderLayerSubsystem.cpp
+++ b/Engine/Plugins/MovieScene/MovieRenderPipeline/Source/MovieRenderPipelineCore/Private/Graph/MovieGraphRenderLayerSubsystem.cpp
@@ -6,6 +6,11 @@
 #include "Components/SkyAtmosphereComponent.h"
 #include "Components/VolumetricCloudComponent.h"
 #include "Components/ExponentialHeightFogComponent.h"
+#include "Engine/DirectionalLight.h"
+#include "Engine/PointLight.h"
+#include "Engine/RectLight.h"
+#include "Engine/SkyLight.h"
+#include "Engine/SpotLight.h"
 #include "EngineUtils.h"
 #include "Framework/Application/SlateApplication.h"
 #include "Materials/MaterialInterface.h"
@@ -1781,6 +1786,119 @@ void UMovieGraphConditionGroupQuery_ActorType::AddActorTypes(const TArray<UClass
 }
 #endif
 
+UMovieGraphConditionGroupQuery_LightActorType::UMovieGraphConditionGroupQuery_LightActorType()
+{
+	// All light actor types are enabled by default.
+	EnabledLightActorTypes = *GetSupportedLightActorTypes();
+}
+
+void UMovieGraphConditionGroupQuery_LightActorType::Evaluate(const TArray<AActor*>& InActorsToQuery, const UWorld* InWorld, TSet<AActor*>& OutMatchingActors) const
+{
+	TRACE_CPUPROFILER_EVENT_SCOPE(UMovieGraphConditionGroupQuery_LightActorType::Evaluate);
+
+	for (AActor* Actor : InActorsToQuery)
+	{
+		if (EnabledLightActorTypes.Contains(Actor->GetClass()))
+		{
+			OutMatchingActors.Add(Actor);
+		}
+	}
+}
+
+const FSlateIcon& UMovieGraphConditionGroupQuery_LightActorType::GetIcon() const
+{
+	static const FSlateIcon LightActorTypeIcon = FSlateIcon(FAppStyle::GetAppStyleSetName(), "ClassIcon.PointLightComponent");
+	return LightActorTypeIcon;
+}
+
+const FText& UMovieGraphConditionGroupQuery_LightActorType::GetDisplayName() const
+{
+	static const FText DisplayName = LOCTEXT("ConditionGroupQueryDisplayName_LightActorType", "Light Actor Type");
+	return DisplayName;
+}
+
+#if WITH_EDITOR
+TArray<TSharedRef<SWidget>> UMovieGraphConditionGroupQuery_LightActorType::GetWidgets()
+{
+	TArray<TSharedRef<SWidget>> Widgets;
+
+	Widgets.Add(
+		SAssignNew(LightActorTypesList, SMovieGraphSimpleList<UClass*>)
+		.DataSource(GetSupportedLightActorTypes())
+		.CustomSummaryWidget(SNullWidget::NullWidget.ToSharedPtr())
+		.OnGetRowText_Static(&GetRowText)
+		.OnGetRowIcon_Static(&GetRowIcon)
+		.ShowEnableDisable(true)
+		.OnGetRowEnableState_Lambda([this](UClass* InLightActorType)
+		{
+			return GetEnableStateForLightType(InLightActorType);
+		})
+		.OnSetRowEnableState_Lambda([this](UClass* InLightActorType, bool bNewEnableState)
+		{
+			SetEnableStateForLightType(InLightActorType, bNewEnableState);
+		})
+	);
+
+	return Widgets;
+}
+#endif	// WITH_EDITOR
+
+void UMovieGraphConditionGroupQuery_LightActorType::SetEnableStateForLightType(UClass* InLightActorType, const bool bNewEnableState)
+{
+	if (GetSupportedLightActorTypes()->Contains(InLightActorType))
+	{
+#if WITH_EDITOR
+		const FScopedTransaction Transaction(LOCTEXT("UpdateLightActorTypesInCollection", "Update Light Actor Types in Collection"));
+		Modify();
+#endif
+
+		if (bNewEnableState)
+		{
+			EnabledLightActorTypes.AddUnique(InLightActorType);
+		}
+		else
+		{
+			EnabledLightActorTypes.Remove(InLightActorType);
+		}
+	}
+}
+
+bool UMovieGraphConditionGroupQuery_LightActorType::GetEnableStateForLightType(UClass* InLightActorType)
+{
+	return EnabledLightActorTypes.Contains(InLightActorType);
+}
+
+TArray<UClass*>* UMovieGraphConditionGroupQuery_LightActorType::GetSupportedLightActorTypes()
+{
+	// Unfortunately this cannot be const because the list widget doesn't expect a const array
+	static TArray<UClass*> SupportedLightActorTypes = {
+		ADirectionalLight::StaticClass(),
+		APointLight::StaticClass(),
+		ARectLight::StaticClass(),
+		ASkyLight::StaticClass(),
+		ASpotLight::StaticClass()
+	};
+
+	return &SupportedLightActorTypes;
+}
+
+#if WITH_EDITOR
+const FSlateBrush* UMovieGraphConditionGroupQuery_LightActorType::GetRowIcon(UClass* InLightActorType)
+{
+	return FSlateIconFinder::FindIconForClass(InLightActorType).GetIcon();
+}
+
+FText UMovieGraphConditionGroupQuery_LightActorType::GetRowText(UClass* InLightActorType)
+{
+	if (InLightActorType)
+	{
+		return InLightActorType->GetDisplayNameText();
+	}
+
+	return LOCTEXT("MovieGraphLightActorTypeConditionGroupQuery_Invalid", "(invalid light type)");
+}
+#endif	// WITH_EDITOR
+
 UMovieGraphConditionGroupQuery_ComponentTagName::UMovieGraphConditionGroupQuery_ComponentTagName()
 	: bOnlyMatchComponents(true)
 {
diff --git a/Engine/Plugins/MovieScene/MovieRenderPipeline/Source/MovieRenderPipelineCore/Public/Graph/MovieGraphRenderLayerSubsystem.h b/Engine/Plugins/MovieScene/MovieRenderPipeline/Source/MovieRenderPipelineCore/Public/Graph/MovieGraphRenderLayerSubsystem.h
--- a/Engine/Plugins/MovieScene/MovieRenderPipeline/Source/MovieRenderPipelineCore/Public/Graph/MovieGraphRenderLayerSubsystem.h
+++ b/Engine/Plugins/MovieScene/MovieRenderPipeline/Source/MovieRenderPipelineCore/Public/Graph/MovieGraphRenderLayerSubsystem.h
@@ -388,6 +388,52 @@ private:
 	TSharedPtr<class SClassViewer> ClassViewerWidget;
 };
 
+/** Query type which filters actors by light type. */
+UCLASS(MinimalAPI, BlueprintType)
+class UMovieGraphConditionGroupQuery_LightActorType final : public UMovieGraphConditionGroupQueryBase
+{
+	GENERATED_BODY()
+
+public:
+	MOVIERENDERPIPELINECORE_API UMovieGraphConditionGroupQuery_LightActorType();
+
+	MOVIERENDERPIPELINECORE_API virtual void Evaluate(const TArray<AActor*>& InActorsToQuery, const UWorld* InWorld, TSet<AActor*>& OutMatchingActors) const override;
+	MOVIERENDERPIPELINECORE_API virtual const FSlateIcon& GetIcon() const override;
+	MOVIERENDERPIPELINECORE_API virtual const FText& GetDisplayName() const override;
+
+#if WITH_EDITOR
+	MOVIERENDERPIPELINECORE_API virtual TArray<TSharedRef<SWidget>> GetWidgets() override;
+#endif
+
+	/**
+	 * Sets the enable state for the given light actor type. All enabled light types will be matched by the query. If the given actor type isn't
+	 * one of the supported light types, this does nothing.
+	 */
+	UFUNCTION(BlueprintCallable, Category = "Lights")
+	void SetEnableStateForLightType(UClass* InLightActorType, const bool bNewEnableState);
+
+	/** Gets the enable state for the given light actor type. If the provided type isn't a supported light type, returns false. */
+	UFUNCTION(BlueprintCallable, Category = "Lights")
+	bool GetEnableStateForLightType(UClass* InLightActorType);
+
+private:
+	/** The light types that are enabled (and will be matched by the query). */
+	UPROPERTY()
+	TArray<TObjectPtr<UClass>> EnabledLightActorTypes;
+
+private:
+	/** Gets the supported light actor types. */
+	static TArray<UClass*>* GetSupportedLightActorTypes();
+
+#if WITH_EDITOR
+	static MOVIERENDERPIPELINECORE_API const FSlateBrush* GetRowIcon(UClass* InLightActorType);
+	static MOVIERENDERPIPELINECORE_API FText GetRowText(UClass* InLightActorType);
+
+	/** The list widget which displays the light actor types and their enable/disable state. */
+	TSharedPtr<SMovieGraphSimpleList<UClass*>> LightActorTypesList;
+#endif
+};
+
 /** Query type which filters actors by tags on their components. */
 UCLASS(MinimalAPI, BlueprintType)
 class UMovieGraphConditionGroupQuery_ComponentTagName final : public UMovieGraphConditionGroupQueryBase
diff --git a/Engine/Plugins/MovieScene/MovieRenderPipeline/Source/MovieRenderPipelineCore/Public/Graph/MovieGraphSharedWidgets.h b/Engine/Plugins/MovieScene/MovieRenderPipeline/Source/MovieRenderPipelineCore/Public/Graph/MovieGraphSharedWidgets.h
--- a/Engine/Plugins/MovieScene/MovieRenderPipeline/Source/MovieRenderPipelineCore/Public/Graph/MovieGraphSharedWidgets.h
+++ b/Engine/Plugins/MovieScene/MovieRenderPipeline/Source/MovieRenderPipelineCore/Public/Graph/MovieGraphSharedWidgets.h
@@ -62,7 +62,8 @@ public:
 
 		/**
 		 * Gets the widget that should be shown in the summary row. By default, "DataType" and "DataTypePlural" will be used to auto-generate the summary
-		 * row widget. If this attribute is bound, those will be ignored, and this will be used instead.
+		 * row widget. If this attribute is bound, those will be ignored, and this will be used instead. To hide the summary row, use SNullWidget::NullWidget
+		 * as the widget.
 		 */
 		SLATE_ATTRIBUTE(TSharedPtr<SWidget>, CustomSummaryWidget);
 
@@ -160,6 +161,13 @@ public:
 			.AutoHeight()
 			[
 				SNew(SBorder)
+				.Visibility_Lambda([this]()
+				{
+					// Hide the summary widget if it's a null widget
+					return CustomSummaryWidget.IsSet() && (CustomSummaryWidget.Get() == SNullWidget::NullWidget.ToSharedPtr())
+						? EVisibility::Collapsed
+						: EVisibility::Visible;
+				})
 				.BorderImage(FAppStyle::Get().GetBrush("Brushes.Header"))
 				.VAlign(VAlign_Center)
 				.HAlign(HAlign_Fill)
@@ -279,12 +287,6 @@ private:
 				.Padding(7.f, 5.f, 7.f, 5.f)
 				[
 					SNew(SCheckBox)
-					.ToolTipText_Lambda([this, InListData]()
-					{
-						return IsRowEnabled(InListData)
-							? LOCTEXT("CollectionEnabled", "This collection is enabled and will be modified by this Modifier node.")
-							: LOCTEXT("CollectionDisabled", "This collection is disabled and will not be modified by this Modifier node.");
-					})
 					.Visibility_Lambda([this]()
 					{
 						return bShowEnableDisable ? EVisibility::Visible : EVisibility::Collapsed;
