MoviePipeline: Added a new "Component Name" condition group query to the graph.

Also renamed the relatively new UMovieGraphClassPropertyModifier to UMovieGraphPropertyModifier to reflect the fact that it can now modify component properties directly.

Additionally, UMovieGraphLightModifierNode was updated to match light components directly rather than light components within specific light actors (eg, the DirectionalLight component generically rather than the DirectionalLight component *within* a Directional Light actor). This means that lighting components can now be matched within blueprints.

#Synced-CL 49732801

diff --git a/Engine/Plugins/MovieScene/MovieRenderPipeline/Source/MovieRenderPipelineCore/Private/Graph/MovieGraphRenderLayerSubsystem.cpp b/Engine/Plugins/MovieScene/MovieRenderPipeline/Source/MovieRenderPipelineCore/Private/Graph/MovieGraphRenderLayerSubsystem.cpp
--- a/Engine/Plugins/MovieScene/MovieRenderPipeline/Source/MovieRenderPipelineCore/Private/Graph/MovieGraphRenderLayerSubsystem.cpp
+++ b/Engine/Plugins/MovieScene/MovieRenderPipeline/Source/MovieRenderPipelineCore/Private/Graph/MovieGraphRenderLayerSubsystem.cpp
@@ -2060,7 +2060,7 @@ TArray<TSharedRef<SWidget>> UMovieGraphConditionGroupQuery_ComponentTagName::Get
 				.Font(FAppStyle::GetFontStyle(TEXT("PropertyWindow.NormalFont")))
 				.Text(LOCTEXT("ConditionGroupQueryComponentTagName_OnlyMatchComponents", "Only Match Components"))
 			]
-			.ToolTipText(LOCTEXT("ConditionGroupQueryComponentTagName_OnlyMatchComponentsTooltip", "If checked, only components of the specified type(s) will be matched. If unchecked, actors that contain components of the specified type(s) will be matched."))
+			.ToolTipText(LOCTEXT("ConditionGroupQueryComponentTagName_OnlyMatchComponentsTooltip", "If checked, only components with the specified tag(s) will be matched. If unchecked, actors that contain components with the specified tag(s) will be matched."))
 		]
 	);
 
@@ -2289,6 +2289,156 @@ void UMovieGraphConditionGroupQuery_ComponentType::Serialize(FArchive& Ar)
 	}
 }
 
+UMovieGraphConditionGroupQuery_ComponentName::UMovieGraphConditionGroupQuery_ComponentName()
+	: bOnlyMatchComponents(true)
+{
+	
+}
+
+void UMovieGraphConditionGroupQuery_ComponentName::Evaluate(
+	const TArray<AActor*>& InActorsToQuery, const UWorld* InWorld, TSet<AActor*>& OutMatchingActors) const
+{
+	if (bOnlyMatchComponents)
+	{
+		UE_LOG(LogMovieRenderPipeline, Warning, TEXT("Evaluate() called on the Component Name condition, but the condition only matches components. No results will be returned; EvaluateActorsAndComponents() should be used instead."));
+	}
+
+	TSet<UActorComponent*> UnusedComponentSet;
+	EvaluateActorsAndComponents(InActorsToQuery, InWorld, OutMatchingActors, UnusedComponentSet);
+}
+
+void UMovieGraphConditionGroupQuery_ComponentName::EvaluateActorsAndComponents(
+	const TArray<AActor*>& InActorsToQuery, const UWorld* InWorld, TSet<AActor*>& OutMatchingActors, TSet<UActorComponent*>& OutMatchingComponents) const
+{
+	TRACE_CPUPROFILER_EVENT_SCOPE(UMovieGraphConditionGroupQuery_ComponentName::EvaluateActorsAndComponents);
+
+	// Quick early-out if "*" is used as the wildcard. Faster than doing the wildcard matching.
+	if (ComponentNames == TEXT("*"))
+	{
+		if (bOnlyMatchComponents)
+		{
+			TInlineComponentArray<UActorComponent*> MatchingComponents;
+
+			// This will be EXTREMELY slow if there are a lot of actors in the scene with many components
+			for (const AActor* ActorToQuery : InActorsToQuery)
+			{
+				constexpr bool bIncludeFromChildActors = true;
+				ActorToQuery->GetComponents(MatchingComponents, bIncludeFromChildActors);
+
+				OutMatchingComponents.Append(MatchingComponents);
+			}
+		}
+		else
+		{
+			OutMatchingActors.Append(InActorsToQuery);
+		}
+		
+		return;
+	}
+
+	// Component names can be specified on multiple lines
+	TArray<FString> AllComponentNames;
+	ComponentNames.ParseIntoArrayLines(AllComponentNames);
+
+	TInlineComponentArray<UActorComponent*> ActorComponents;
+
+	for (AActor* Actor : InActorsToQuery)
+	{
+		// Include child components so components inside of Blueprints can be found
+		constexpr bool bIncludeFromChildActors = true;
+		Actor->GetComponents<UActorComponent*>(ActorComponents, bIncludeFromChildActors);
+
+		for (UActorComponent* Component : ActorComponents)
+		{
+			for (const FString& ComponentName : AllComponentNames)
+			{
+				if (Component->GetName().MatchesWildcard(ComponentName))
+				{
+					if (bOnlyMatchComponents)
+					{
+						OutMatchingComponents.Add(Component);
+					}
+					else
+					{
+						OutMatchingActors.Add(Actor);
+					}
+				}
+			}
+		}
+
+		ActorComponents.Reset();
+	}
+}
+
+bool UMovieGraphConditionGroupQuery_ComponentName::ShouldEvaluateComponents() const
+{
+	return bOnlyMatchComponents;
+}
+
+const FSlateIcon& UMovieGraphConditionGroupQuery_ComponentName::GetIcon() const
+{
+	static const FSlateIcon ComponentNameIcon = FSlateIcon(FAppStyle::GetAppStyleSetName(), "ClassIcon.TextRenderActor");
+	return ComponentNameIcon;
+}
+
+const FText& UMovieGraphConditionGroupQuery_ComponentName::GetDisplayName() const
+{
+	static const FText DisplayName = LOCTEXT("ConditionGroupQueryDisplayName_ComponentName", "Component Name");
+	return DisplayName;
+}
+
+#if WITH_EDITOR
+TArray<TSharedRef<SWidget>> UMovieGraphConditionGroupQuery_ComponentName::GetWidgets()
+{
+	TArray<TSharedRef<SWidget>> Widgets;
+
+	Widgets.Add(
+		SNew(SVerticalBox)
+
+		+ SVerticalBox::Slot()
+		.HAlign(HAlign_Fill)
+		.AutoHeight()
+		.Padding(7.f, 2.f)
+		[
+			SNew(SMultiLineEditableTextBox)
+			.Text_Lambda([this]() { return FText::FromString(ComponentNames); })
+			.OnTextCommitted_Lambda([this](const FText& InText, ETextCommit::Type TextCommitType)
+			{
+				const FScopedTransaction Transaction(LOCTEXT("UpdateComponentNamesInCollection", "Update Component Names in Collection"));
+				Modify();
+
+				ComponentNames = InText.ToString();
+			})
+			.HintText(LOCTEXT("MovieGraphComponentNameQueryHintText", "A component on the actor must match one or more component names.\nWildcards allowed. Enter each component name on a separate line."))
+		]
+
+		+ SVerticalBox::Slot()
+		.HAlign(HAlign_Right)
+		.AutoHeight()
+		.Padding(8, 0, 8, 2)
+		[
+			SNew(SCheckBox)
+			.OnCheckStateChanged_Lambda([this](ECheckBoxState NewState)
+			{
+				const FScopedTransaction Transaction(LOCTEXT("ConditionGroupQueryComponentName_ToggleOnlyMatchComponents", "Toggle Only Match Components"));
+				Modify();
+
+				bOnlyMatchComponents = (NewState == ECheckBoxState::Checked);
+			})
+			.IsChecked_Lambda([this]() { return bOnlyMatchComponents ? ECheckBoxState::Checked : ECheckBoxState::Unchecked; })
+			[
+				SNew(STextBlock)
+				.Font(FAppStyle::GetFontStyle(TEXT("PropertyWindow.NormalFont")))
+				.Text(LOCTEXT("ConditionGroupQueryComponentName_OnlyMatchComponents", "Only Match Components"))
+			]
+			.ToolTipText(LOCTEXT("ConditionGroupQueryComponentName_OnlyMatchComponentsTooltip", "If checked, only components with the specified name(s) will be matched. If unchecked, actors that contain components with the specified name(s) will be matched."))
+		]
+	);
+
+	return Widgets;
+}
+#endif	// WITH_EDITOR
+
 void UMovieGraphConditionGroupQuery_EditorFolder::Evaluate(const TArray<AActor*>& InActorsToQuery, const UWorld* InWorld, TSet<AActor*>& OutMatchingActors) const
 {
 #if WITH_EDITOR
diff --git a/Engine/Plugins/MovieScene/MovieRenderPipeline/Source/MovieRenderPipelineCore/Private/Graph/Nodes/MovieGraphClassPropertyModifier.cpp b/Engine/Plugins/MovieScene/MovieRenderPipeline/Source/MovieRenderPipelineCore/Private/Graph/Nodes/MovieGraphClassPropertyModifier.cpp
--- a/Engine/Plugins/MovieScene/MovieRenderPipeline/Source/MovieRenderPipelineCore/Private/Graph/Nodes/MovieGraphClassPropertyModifier.cpp
+++ b/Engine/Plugins/MovieScene/MovieRenderPipeline/Source/MovieRenderPipelineCore/Private/Graph/Nodes/MovieGraphClassPropertyModifier.cpp
@@ -53,7 +53,7 @@ TSubclassOf<UActorComponent> FMovieGraphPropertyReference::GetComponentClass() c
 
 bool FMovieGraphPropertyReference::IsValidProperty() const
 {
-	return !!Property && PropertyClass.Get() && !PropertyName.IsNone();
+	return !!Property && (PropertyClass.Get() || ComponentClass.Get()) && !PropertyName.IsNone();
 }
 
 void FMovieGraphPropertyReference::PostSerialize(const FArchive& Ar)
@@ -91,7 +91,7 @@ void FMovieGraphPropertyReference::PopulateDerivedData()
 	}
 }
 
-UMovieGraphClassPropertyModifier::UMovieGraphClassPropertyModifier()
+UMovieGraphPropertyModifier::UMovieGraphPropertyModifier()
 {
 	if (!HasAnyFlags(RF_ClassDefaultObject))
 	{
@@ -103,8 +103,27 @@ UMovieGraphClassPropertyModifier::UMovieGraphClassPropertyModifier()
 	}
 }
 
-void UMovieGraphClassPropertyModifier::ApplyModifier(const UWorld* World)
+void UMovieGraphPropertyModifier::ApplyModifier(const UWorld* World)
 {
+	// Set the value of the given property (with the value that was set for it on the modifier), and also cache its current value so it can be
+	// restored when the modifier is undone. The property is modified within the specified container (an actor or component).
+	auto CacheAndSetPropertyValue = [this](const FMovieGraphPropertyReference& InPropertyReference, UObject* InPropertyContainer)
+	{
+		// TODO: Ideally the below export/import aren't string based, but that would require a lot of specialized code in order to
+		// properly handle every type. This will do for now, but could potentially be optimized in the future if this is a bottleneck
+		// or causes issues.
+
+		// Cache the current value so it can be restored later
+		FString CurrentValue;
+		InPropertyReference.GetProperty()->ExportTextItem_InContainer(CurrentValue, InPropertyContainer, nullptr, nullptr, PPF_None);
+		PreviousPropertyValuesView->SetValueSerializedString(InPropertyReference.AuxPropertyName, CurrentValue);
+		ModifiedProperties.Add({ InPropertyContainer, InPropertyReference });
+
+		// Set the property to its new value
+		const FString PropertyValue = UpdatedPropertyValuesView->GetValueSerializedString(InPropertyReference.AuxPropertyName);
+		InPropertyReference.GetProperty()->ImportText_InContainer(*PropertyValue, InPropertyContainer, InPropertyContainer, 0);
+	};
+
 	TArray<TSubclassOf<AActor>> TargetActorClasses;
 	PropertiesToModifyByActorClass.GetKeys(TargetActorClasses);
 
@@ -120,48 +139,84 @@ void UMovieGraphClassPropertyModifier::ApplyModifier(const UWorld* World)
 
 		FMovieGraphEvaluationResult EvaluationResult = Collection->EvaluateActorsAndComponents(World);
 
+		// Set any actor properties on actors the collection contains.
 		for (const TObjectPtr<AActor>& MatchingActor : EvaluationResult.MatchingActors)
 		{
 			UClass* ActorClass = MatchingActor->GetClass();
+			bool bOnlyMatchComponent = false;
 
-			// The modifier only applies to specific actor types
+			// The modifier only applies to specific actor types. If this modifier isn't targeting any of the actor types in the collection, the
+			// modifier *may* indicate that some properties should be applied to all components of a specific type within an actor (which is indicated
+			// by a null target actor class).
 			if (!TargetActorClasses.Contains(ActorClass))
 			{
-				continue;
+				bOnlyMatchComponent = true;
+				ActorClass = nullptr;
+
+				// However, there may be component-specific matches that don't target distinct actor classes
+				if (!TargetActorClasses.Contains(nullptr))
+				{
+					continue;
+				}
 			}
 
 			for (const FMovieGraphPropertyReference& PropertyReference : PropertiesToModifyByActorClass[ActorClass])
 			{
-				// If the property is part of a component, use the component as the container property; otherwise, use the actor itself.
+				// If the property is part of a component, use the component as the property container; otherwise, use the actor itself.
 				UObject* PropertyContainer = nullptr;
 				if (PropertyReference.GetComponentClass())
 				{
 					PropertyContainer = MatchingActor->GetComponentByClass(PropertyReference.GetComponentClass());
 					DirtiedComponents.Add(Cast<UActorComponent>(PropertyContainer));
 				}
-				else
+				else if (!bOnlyMatchComponent)
 				{
 					PropertyContainer = MatchingActor;
 				}
 
-				if (!PropertyContainer)
+				if (PropertyContainer)
 				{
-					continue;
+					CacheAndSetPropertyValue(PropertyReference, PropertyContainer);
 				}
+			}
+		}
 
-				// TODO: Ideally the below export/import aren't string based, but that would require a lot of specialized code in order to
-				// properly handle every type. This will do for now, but could potentially be optimized in the future if this is a bottleneck
-				// or causes issues.
-
-				// Cache the current value so it can be restored later
-				FString CurrentValue;
-				PropertyReference.GetProperty()->ExportTextItem_InContainer(CurrentValue, PropertyContainer, nullptr, nullptr, PPF_None);
-				PreviousPropertyValuesView->SetValueSerializedString(PropertyReference.AuxPropertyName, CurrentValue);
-				ModifiedProperties.Add({ PropertyContainer, PropertyReference });
-				
-				// Set the property to its new value
-				const FString PropertyValue = UpdatedPropertyValuesView->GetValueSerializedString(PropertyReference.AuxPropertyName);
-				PropertyReference.GetProperty()->ImportText_InContainer(*PropertyValue, PropertyContainer, PropertyContainer, 0);
+		// Set any component properties on components the collection contains.
+		for (const TObjectPtr<UActorComponent>& MatchingComponent : EvaluationResult.MatchingComponents)
+		{
+			// The properties to modify are grouped by actor type. For the purposes of components, we need to inspect all actor types since any
+			// actor type could potentially contain a component of the type(s) that are contained in the evaluated collection.
+			for (const TPair<TSubclassOf<AActor>, TArray<FMovieGraphPropertyReference>>& PropertyPair : PropertiesToModifyByActorClass)
+			{
+				const UClass* ActorType = PropertyPair.Key.Get();
+				const TArray<FMovieGraphPropertyReference>& PropertyReferences = PropertyPair.Value;
+
+				for (const FMovieGraphPropertyReference& PropertyReference : PropertyReferences)
+				{
+					UObject* PropertyContainer = nullptr;
+					if (PropertyReference.GetComponentClass() == MatchingComponent->GetClass())
+					{
+						// The property reference may indicate that the component AND class types need to match
+						if (PropertyReference.GetPropertyClass())
+						{
+							if (MatchingComponent->GetOwner()->GetClass() == ActorType)
+							{
+								PropertyContainer = MatchingComponent;
+							}
+						}
+						else
+						{
+							PropertyContainer = MatchingComponent;
+						}
+
+						DirtiedComponents.Add(Cast<UActorComponent>(PropertyContainer));
+					}
+
+					if (PropertyContainer)
+					{
+						CacheAndSetPropertyValue(PropertyReference, PropertyContainer);
+					}
+				}
 			}
 		}
 	}
@@ -176,7 +231,7 @@ void UMovieGraphClassPropertyModifier::ApplyModifier(const UWorld* World)
 	}
 }
 
-void UMovieGraphClassPropertyModifier::UndoModifier()
+void UMovieGraphPropertyModifier::UndoModifier()
 {
 	for (const TPair<UObject*, FMovieGraphPropertyReference>& ModifiedPropertyPair : ModifiedProperties)
 	{
@@ -201,13 +256,13 @@ void UMovieGraphClassPropertyModifier::UndoModifier()
 	ModifiedProperties.Reset();
 }
 
-FText UMovieGraphClassPropertyModifier::GetModifierName()
+FText UMovieGraphPropertyModifier::GetModifierName()
 {
 	static const FText ModifierName = NSLOCTEXT("MovieGraph", "ClassPropertyModifierName", "Class Property");
 	return ModifierName;
 }
 
-FMovieGraphPropertyReference UMovieGraphClassPropertyModifier::MakePropertyReference(
+FMovieGraphPropertyReference UMovieGraphPropertyModifier::MakePropertyReference(
 	UClass* InPropertyClass,
 	const TSubclassOf<UActorComponent>& InComponentClass,
 	const FName& InPropertyName)
@@ -215,7 +270,7 @@ FMovieGraphPropertyReference UMovieGraphClassPropertyModifier::MakePropertyRefer
 	return FMovieGraphPropertyReference(InPropertyClass, InComponentClass, InPropertyName);
 }
 
-void UMovieGraphClassPropertyModifier::AddProperties(const TArray<FMovieGraphPropertyReference>& InProperties)
+void UMovieGraphPropertyModifier::AddProperties(const TArray<FMovieGraphPropertyReference>& InProperties)
 {
 	TArray<FMovieGraphValueViewProperty> PropertiesToAdd;
 
@@ -245,12 +300,12 @@ void UMovieGraphClassPropertyModifier::AddProperties(const TArray<FMovieGraphPro
 	PreviousPropertyValuesView->AddProperties(PropertiesToAdd);
 }
 
-FName UMovieGraphClassPropertyModifier::GetPropertyNameForSettingCustomValue(const FMovieGraphPropertyReference& InPropertyReference)
+FName UMovieGraphPropertyModifier::GetPropertyNameForSettingCustomValue(const FMovieGraphPropertyReference& InPropertyReference)
 {
 	return GetInternalPropertyName(InPropertyReference);
 }
 
-UMovieGraphFixedValueView* UMovieGraphClassPropertyModifier::GetPropertiesForSettingValues()
+UMovieGraphFixedValueView* UMovieGraphPropertyModifier::GetPropertiesForSettingValues()
 {
 	if (!UserFacingPropertiesView)
 	{
@@ -261,7 +316,7 @@ UMovieGraphFixedValueView* UMovieGraphClassPropertyModifier::GetPropertiesForSet
 	return UserFacingPropertiesView;
 }
 
-FName UMovieGraphClassPropertyModifier::GetInternalPropertyName(const FMovieGraphPropertyReference& InPropertyReference)
+FName UMovieGraphPropertyModifier::GetInternalPropertyName(const FMovieGraphPropertyReference& InPropertyReference)
 {
 	// If the property reference includes a component, that needs to be included in the property name in order to properly disambiguate it.
 	// Otherwise, using the class name is sufficient.
@@ -281,6 +336,13 @@ FName UMovieGraphClassPropertyModifier::GetInternalPropertyName(const FMovieGrap
 			InPropertyReference.GetPropertyName().ToString() }));
 	}
 
+	if (InPropertyReference.GetComponentClass())
+	{
+		return FName(FString::Format(TEXT("{0}_{1}"), {
+			InPropertyReference.GetComponentClass()->GetName(),
+			InPropertyReference.GetPropertyName().ToString() }));
+	}
+
 	UE_LOG(LogMovieRenderPipeline, Warning, TEXT("Attempting to modify a property via the Property Modifier with an invalid actor class and/or component class."));
 
 	return FName(TEXT("INVALID"));
diff --git a/Engine/Plugins/MovieScene/MovieRenderPipeline/Source/MovieRenderPipelineCore/Private/Graph/Nodes/MovieGraphLightModifierNode.cpp b/Engine/Plugins/MovieScene/MovieRenderPipeline/Source/MovieRenderPipelineCore/Private/Graph/Nodes/MovieGraphLightModifierNode.cpp
--- a/Engine/Plugins/MovieScene/MovieRenderPipeline/Source/MovieRenderPipelineCore/Private/Graph/Nodes/MovieGraphLightModifierNode.cpp
+++ b/Engine/Plugins/MovieScene/MovieRenderPipeline/Source/MovieRenderPipelineCore/Private/Graph/Nodes/MovieGraphLightModifierNode.cpp
@@ -2,14 +2,14 @@
 
 #include "Graph/Nodes/MovieGraphLightModifierNode.h"
 
+#include "Components/DirectionalLightComponent.h"
 #include "Components/LightComponent.h"
 #include "Components/LightComponentBase.h"
 #include "Components/LocalLightComponent.h"
-#include "Engine/DirectionalLight.h"
-#include "Engine/PointLight.h"
-#include "Engine/RectLight.h"
-#include "Engine/SkyLight.h"
-#include "Engine/SpotLight.h"
+#include "Components/PointLightComponent.h"
+#include "Components/RectLightComponent.h"
+#include "Components/SkyLightComponent.h"
+#include "Components/SpotLightComponent.h"
 #include "GameFramework/Actor.h"
 #include "Graph/MovieGraphValueUtils.h"
 #include "Graph/Nodes/MovieGraphClassPropertyModifier.h"
@@ -20,7 +20,7 @@
 
 UMovieGraphLightModifierNode::UMovieGraphLightModifierNode()
 {
-	Modifier = CreateDefaultSubobject<UMovieGraphClassPropertyModifier>(TEXT("LightingModifier"));
+	Modifier = CreateDefaultSubobject<UMovieGraphPropertyModifier>(TEXT("LightingModifier"));
 }
 
 #if WITH_EDITOR
@@ -223,7 +223,7 @@ TArray<UMovieGraphModifierBase*> UMovieGraphLightModifierNode::GetAllModifiers()
 	TArray<TFunction<void()>> DeferredValueSetters;
 
 	auto AddClassPredefinedLightProperty = [this, &ModifierValues, &PropertiesToAdd, &DeferredValueSetters]
-		(const bool& InOverrideProperty, const FName& InComponentPropertyName, const FName& InNodePropertyName, UClass* InLightClass, UClass* InComponentClass) -> void
+		(const bool& InOverrideProperty, const FName& InComponentPropertyName, const FName& InNodePropertyName, UClass* InComponentClass) -> void
 	{
 		// Nothing to do if this lighting property wasn't marked as overridden.
 		if (!InOverrideProperty)
@@ -238,8 +238,9 @@ TArray<UMovieGraphModifierBase*> UMovieGraphLightModifierNode::GetAllModifiers()
 			return;
 		}
 
+		constexpr UClass* LightClass = nullptr;
 		FMovieGraphPropertyReference& NewProperty =
-			PropertiesToAdd.Add_GetRef(FMovieGraphPropertyReference(InLightClass, InComponentClass, InComponentPropertyName, ComponentProperty));
+			PropertiesToAdd.Add_GetRef(FMovieGraphPropertyReference(LightClass, InComponentClass, InComponentPropertyName, ComponentProperty));
 
 		// Queue up setting the value in the property bag *after* the property was added to the property bag.
 		DeferredValueSetters.Add([this, &ModifierValues, NewProperty, InNodePropertyName]()
@@ -257,50 +258,51 @@ TArray<UMovieGraphModifierBase*> UMovieGraphLightModifierNode::GetAllModifiers()
 
 	/** Adds a lighting property to the modifier. This propery applies to all light types. */
 	auto AddGeneralPredefinedLightProperty = [this, &AddClassPredefinedLightProperty]
-		(const bool& InOverrideProperty, const FName& InComponentPropertyName, const FName& InNodePropertyName, UClass* InComponentClass) -> void
+		(const bool& InOverrideProperty, const FName& InComponentPropertyName, const FName& InNodePropertyName) -> void
 	{
-		static const TArray<UClass*> LightingClasses = {
-			ADirectionalLight::StaticClass(), ASkyLight::StaticClass(), APointLight::StaticClass(), ASpotLight::StaticClass(), ARectLight::StaticClass()
+		static const TArray<UClass*> ComponentClasses = {
+			UDirectionalLightComponent::StaticClass(), USkyLightComponent::StaticClass(), UPointLightComponent::StaticClass(),
+			USpotLightComponent::StaticClass(), URectLightComponent::StaticClass()
 		};
 
-		// Add this property for all light classes.
-		for (UClass* LightingClass : LightingClasses)
+		// Add this property for all light component classes.
+		for (UClass* ComponentClass : ComponentClasses)
 		{
-			AddClassPredefinedLightProperty(InOverrideProperty, InComponentPropertyName, InNodePropertyName, LightingClass, InComponentClass);
+			AddClassPredefinedLightProperty(InOverrideProperty, InComponentPropertyName, InNodePropertyName, ComponentClass);
 		}
 	};
 
 	// First, add in the properties that are always defined on the node.
-	AddGeneralPredefinedLightProperty(bOverride_LightColor, LightColorName, LightColorName, ULightComponentBase::StaticClass());
-	AddGeneralPredefinedLightProperty(bOverride_bAffectsWorld, AffectsWorldName, AffectsWorldName, ULightComponentBase::StaticClass());
-	AddGeneralPredefinedLightProperty(bOverride_CastShadows, CastShadowsName, CastShadowsName, ULightComponentBase::StaticClass());
-	AddGeneralPredefinedLightProperty(bOverride_IndirectLightingIntensity, IndirectLightingIntensityName, IndirectLightingIntensityName, ULightComponentBase::StaticClass());
-	AddGeneralPredefinedLightProperty(bOverride_VolumetricScatteringIntensity, VolumetricScatteringIntensityName, VolumetricScatteringIntensityName, ULightComponentBase::StaticClass());
-	AddGeneralPredefinedLightProperty(bOverride_LightingChannels, LightingChannelsName, LightingChannelsName, ULightComponent::StaticClass());
-	AddGeneralPredefinedLightProperty(bOverride_SamplesPerPixel, SamplesPerPixelName, SamplesPerPixelName, ULightComponentBase::StaticClass());
-	AddClassPredefinedLightProperty(bOverride_DirectionalLightIntensity, IntensityName, DirectionalLightIntensityName, ADirectionalLight::StaticClass(), ULightComponentBase::StaticClass());
-	AddClassPredefinedLightProperty(bOverride_SkyLightIntensityScale, IntensityName, SkyLightIntensityName, ASkyLight::StaticClass(), ULightComponentBase::StaticClass());
+	AddGeneralPredefinedLightProperty(bOverride_LightColor, LightColorName, LightColorName);
+	AddGeneralPredefinedLightProperty(bOverride_bAffectsWorld, AffectsWorldName, AffectsWorldName);
+	AddGeneralPredefinedLightProperty(bOverride_CastShadows, CastShadowsName, CastShadowsName);
+	AddGeneralPredefinedLightProperty(bOverride_IndirectLightingIntensity, IndirectLightingIntensityName, IndirectLightingIntensityName);
+	AddGeneralPredefinedLightProperty(bOverride_VolumetricScatteringIntensity, VolumetricScatteringIntensityName, VolumetricScatteringIntensityName);
+	AddGeneralPredefinedLightProperty(bOverride_LightingChannels, LightingChannelsName, LightingChannelsName);
+	AddGeneralPredefinedLightProperty(bOverride_SamplesPerPixel, SamplesPerPixelName, SamplesPerPixelName);
+	AddClassPredefinedLightProperty(bOverride_DirectionalLightIntensity, IntensityName, DirectionalLightIntensityName, UDirectionalLightComponent::StaticClass());
+	AddClassPredefinedLightProperty(bOverride_SkyLightIntensityScale, IntensityName, SkyLightIntensityName, USkyLightComponent::StaticClass());
 
 	if (IntensityMethod == EMovieGraphLightModifierIntensityMethod::PointRectSpot)
 	{
-		AddClassPredefinedLightProperty(bOverride_Intensity, IntensityName, IntensityName, APointLight::StaticClass(), ULightComponentBase::StaticClass());
-		AddClassPredefinedLightProperty(bOverride_Intensity, IntensityName, IntensityName, ARectLight::StaticClass(), ULightComponentBase::StaticClass());
-		AddClassPredefinedLightProperty(bOverride_Intensity, IntensityName, IntensityName, ASpotLight::StaticClass(), ULightComponentBase::StaticClass());
-		
-		AddClassPredefinedLightProperty(bOverride_IntensityUnits, IntensityUnitsName, IntensityUnitsName, APointLight::StaticClass(), ULocalLightComponent::StaticClass());
-		AddClassPredefinedLightProperty(bOverride_IntensityUnits, IntensityUnitsName, IntensityUnitsName, ARectLight::StaticClass(), ULocalLightComponent::StaticClass());
-		AddClassPredefinedLightProperty(bOverride_IntensityUnits, IntensityUnitsName, IntensityUnitsName, ASpotLight::StaticClass(), ULocalLightComponent::StaticClass());
+		AddClassPredefinedLightProperty(bOverride_Intensity, IntensityName, IntensityName, UPointLightComponent::StaticClass());
+		AddClassPredefinedLightProperty(bOverride_Intensity, IntensityName, IntensityName, URectLightComponent::StaticClass());
+		AddClassPredefinedLightProperty(bOverride_Intensity, IntensityName, IntensityName, USpotLightComponent::StaticClass());
+
+		AddClassPredefinedLightProperty(bOverride_IntensityUnits, IntensityUnitsName, IntensityUnitsName, UPointLightComponent::StaticClass());
+		AddClassPredefinedLightProperty(bOverride_IntensityUnits, IntensityUnitsName, IntensityUnitsName, URectLightComponent::StaticClass());
+		AddClassPredefinedLightProperty(bOverride_IntensityUnits, IntensityUnitsName, IntensityUnitsName, USpotLightComponent::StaticClass());
 	}
 	else
 	{
-		AddClassPredefinedLightProperty(bOverride_PointLightIntensity, IntensityName, PointLightIntensityName, APointLight::StaticClass(), ULightComponentBase::StaticClass());
-		AddClassPredefinedLightProperty(bOverride_PointLightIntensityUnits, IntensityUnitsName, PointLightIntensityUnitsName, APointLight::StaticClass(), ULocalLightComponent::StaticClass());
+		AddClassPredefinedLightProperty(bOverride_PointLightIntensity, IntensityName, PointLightIntensityName, UPointLightComponent::StaticClass());
+		AddClassPredefinedLightProperty(bOverride_PointLightIntensityUnits, IntensityUnitsName, PointLightIntensityUnitsName, UPointLightComponent::StaticClass());
 
-		AddClassPredefinedLightProperty(bOverride_RectLightIntensity, IntensityName, RectLightIntensityName, ARectLight::StaticClass(), ULightComponentBase::StaticClass());
-		AddClassPredefinedLightProperty(bOverride_RectLightIntensityUnits, IntensityUnitsName, RectLightIntensityUnitsName, ARectLight::StaticClass(), ULocalLightComponent::StaticClass());
+		AddClassPredefinedLightProperty(bOverride_RectLightIntensity, IntensityName, RectLightIntensityName, URectLightComponent::StaticClass());
+		AddClassPredefinedLightProperty(bOverride_RectLightIntensityUnits, IntensityUnitsName, RectLightIntensityUnitsName, URectLightComponent::StaticClass());
 
-		AddClassPredefinedLightProperty(bOverride_SpotLightIntensity, IntensityName, SpotLightIntensityName, ASpotLight::StaticClass(), ULightComponentBase::StaticClass());
-		AddClassPredefinedLightProperty(bOverride_SpotLightIntensityUnits, IntensityUnitsName, SpotLightIntensityUnitsName, ASpotLight::StaticClass(), ULocalLightComponent::StaticClass());
+		AddClassPredefinedLightProperty(bOverride_SpotLightIntensity, IntensityName, SpotLightIntensityName, USpotLightComponent::StaticClass());
+		AddClassPredefinedLightProperty(bOverride_SpotLightIntensityUnits, IntensityUnitsName, SpotLightIntensityUnitsName, USpotLightComponent::StaticClass());
 	}
 
 	const UMovieGraphMutableValueView* CustomLightingValues = GetDynamicPropertiesView();
@@ -404,13 +406,14 @@ bool UMovieGraphLightModifierNode::AddCustomLightProperty(UClass* InLightClass,
 	Modify();
 #endif
 
-	// Add supplemental property information for this property, which will importantly track the origin class and component.
+	// Add supplemental property information for this property, which will importantly track the origin component.
 	// This has to be done via sidecar data, rather than via metadata on the property in the property bag, because runtime builds do not contain
 	// property metadata.
+	constexpr UClass* LightClass = nullptr;
 	FMovieGraphPropertyReference& NewCustomProperty = CustomPropertyInfo.Emplace(
 		PropertyName,
-		FMovieGraphPropertyReference(InLightClass, InLightingProperty->GetOwnerClass(), InLightingProperty->GetFName(), const_cast<FProperty*>(InLightingProperty)));
-	NewCustomProperty.AuxPropertyName = PropertyName;	
+		FMovieGraphPropertyReference(LightClass, InLightingProperty->GetOwnerClass(), InLightingProperty->GetFName(), const_cast<FProperty*>(InLightingProperty)));
+	NewCustomProperty.AuxPropertyName = PropertyName;
 
 	// Update DynamicProperties with the properties referenced in CustomPropertyInfo. Although we could update DynamicProperties right here, there are
 	// parts of the graph evaluation lifecycle that depend on GetDynamicPropertyDescriptions() reflecting the property layout of DynamicProperties.
diff --git a/Engine/Plugins/MovieScene/MovieRenderPipeline/Source/MovieRenderPipelineCore/Public/Graph/MovieGraphRenderLayerSubsystem.h b/Engine/Plugins/MovieScene/MovieRenderPipeline/Source/MovieRenderPipelineCore/Public/Graph/MovieGraphRenderLayerSubsystem.h
--- a/Engine/Plugins/MovieScene/MovieRenderPipeline/Source/MovieRenderPipelineCore/Public/Graph/MovieGraphRenderLayerSubsystem.h
+++ b/Engine/Plugins/MovieScene/MovieRenderPipeline/Source/MovieRenderPipelineCore/Public/Graph/MovieGraphRenderLayerSubsystem.h
@@ -479,7 +479,7 @@ public:
 	bool bOnlyMatchComponents;
 };
 
-/** Query type which filters actors via the components contained in them. */
+/** Query type which filters actors via the component types contained in them. */
 UCLASS(MinimalAPI, BlueprintType)
 class UMovieGraphConditionGroupQuery_ComponentType final : public UMovieGraphConditionGroupQueryBase
 {
@@ -526,6 +526,35 @@ public:
 	bool bOnlyMatchComponents;
 };
 
+/** Query type which filters actors via the components contained in them (by name). */
+UCLASS(MinimalAPI, BlueprintType)
+class UMovieGraphConditionGroupQuery_ComponentName final : public UMovieGraphConditionGroupQueryBase
+{
+	GENERATED_BODY()
+
+public:
+	MOVIERENDERPIPELINECORE_API UMovieGraphConditionGroupQuery_ComponentName();
+
+	MOVIERENDERPIPELINECORE_API virtual void Evaluate(const TArray<AActor*>& InActorsToQuery, const UWorld* InWorld, TSet<AActor*>& OutMatchingActors) const override;
+	MOVIERENDERPIPELINECORE_API virtual void EvaluateActorsAndComponents(const TArray<AActor*>& InActorsToQuery, const UWorld* InWorld, TSet<AActor*>& OutMatchingActors, TSet<UActorComponent*>& OutMatchingComponents) const override;
+	MOVIERENDERPIPELINECORE_API virtual bool ShouldEvaluateComponents() const override;
+	MOVIERENDERPIPELINECORE_API virtual const FSlateIcon& GetIcon() const override;
+	MOVIERENDERPIPELINECORE_API virtual const FText& GetDisplayName() const override;
+
+#if WITH_EDITOR
+	MOVIERENDERPIPELINECORE_API virtual TArray<TSharedRef<SWidget>> GetWidgets() override;
+#endif
+
+public:
+	/** The actor must have one or more of the component name(s) in order to be a match. */
+	UPROPERTY(EditAnywhere, Category="General")
+	FString ComponentNames;
+
+	/** Whether this should match components or actors. If false, any components that match will instead match their parent actor. */
+	UPROPERTY(EditAnywhere, Category="General")
+	bool bOnlyMatchComponents;
+};
+
 /** Query type which filters actors via the editor folder that they're contained in. */
 UCLASS(MinimalAPI, BlueprintType)
 class UMovieGraphConditionGroupQuery_EditorFolder final : public UMovieGraphConditionGroupQueryBase
diff --git a/Engine/Plugins/MovieScene/MovieRenderPipeline/Source/MovieRenderPipelineCore/Public/Graph/Nodes/MovieGraphClassPropertyModifier.h b/Engine/Plugins/MovieScene/MovieRenderPipeline/Source/MovieRenderPipelineCore/Public/Graph/Nodes/MovieGraphClassPropertyModifier.h
--- a/Engine/Plugins/MovieScene/MovieRenderPipeline/Source/MovieRenderPipelineCore/Public/Graph/Nodes/MovieGraphClassPropertyModifier.h
+++ b/Engine/Plugins/MovieScene/MovieRenderPipeline/Source/MovieRenderPipelineCore/Public/Graph/Nodes/MovieGraphClassPropertyModifier.h
@@ -101,19 +101,19 @@ struct TStructOpsTypeTraits<FMovieGraphPropertyReference> : public TStructOpsTyp
 };
 
 /**
- * Modifies properties on specific actor type(s).
+ * Modifies properties on specific actor or component type(s).
  *
  * The properties that should be modified are added via MovieGraphPropertyReference objects via AddProperties(). Note that the
  * properties are matched to their EXACT AActor and/or UActorComponent class within the property reference (ie, subclass checks
  * are not made).
  */
 UCLASS(MinimalAPI, BlueprintType)
-class UMovieGraphClassPropertyModifier final : public UMovieGraphCollectionModifier
+class UMovieGraphPropertyModifier final : public UMovieGraphCollectionModifier
 {
 	GENERATED_BODY()
 
 public:
-	UMovieGraphClassPropertyModifier();
+	UMovieGraphPropertyModifier();
 
 	//~ Begin UMovieGraphModifierBase Interface
 	UE_API virtual void ApplyModifier(const UWorld* World) override;
@@ -131,7 +131,14 @@ public:
 		const TSubclassOf<UActorComponent>& InComponentClass,
 		const FName& InPropertyName);
 
-	/** Adds the given properties (which are each associated with a specific class). */
+	/**
+	 * Adds the properties which should be modified.
+	 *
+	 * The property reference can:
+	 * 1) Specify an actor without its component to modify an actor's property.
+	 * 2) Specify a component with an actor to modify that component's property only on specific actor types.
+	 * 3) Specify a component without an actor type to modify that component property on all actors it exists on.
+	 */
 	UFUNCTION(BlueprintCallable, Category = "Properties")
 	UE_API void AddProperties(const TArray<FMovieGraphPropertyReference>& InProperties);
 
diff --git a/Engine/Plugins/MovieScene/MovieRenderPipeline/Source/MovieRenderPipelineCore/Public/Graph/Nodes/MovieGraphLightModifierNode.h b/Engine/Plugins/MovieScene/MovieRenderPipeline/Source/MovieRenderPipelineCore/Public/Graph/Nodes/MovieGraphLightModifierNode.h
--- a/Engine/Plugins/MovieScene/MovieRenderPipeline/Source/MovieRenderPipelineCore/Public/Graph/Nodes/MovieGraphLightModifierNode.h
+++ b/Engine/Plugins/MovieScene/MovieRenderPipeline/Source/MovieRenderPipelineCore/Public/Graph/Nodes/MovieGraphLightModifierNode.h
@@ -9,7 +9,7 @@
 
 #include "MovieGraphLightModifierNode.generated.h"
 
-class UMovieGraphClassPropertyModifier;
+class UMovieGraphPropertyModifier;
 
 #define UE_API MOVIERENDERPIPELINECORE_API
 
@@ -335,7 +335,7 @@ private:
 
 	/** The light modifier associated with this node. */
 	UPROPERTY(meta = (HideInActiveRenderSettings))
-	TObjectPtr<UMovieGraphClassPropertyModifier> Modifier;
+	TObjectPtr<UMovieGraphPropertyModifier> Modifier;
 
 	/**
 	 * Stores supplemental information for each "Custom" lighting property. Key is the name of the corresponding non-override property
diff --git a/Engine/Plugins/MovieScene/MovieRenderPipeline/Source/MovieRenderPipelineEditor/Private/Customizations/Graph/MovieGraphLightModifierCustomization.cpp b/Engine/Plugins/MovieScene/MovieRenderPipeline/Source/MovieRenderPipelineEditor/Private/Customizations/Graph/MovieGraphLightModifierCustomization.cpp
--- a/Engine/Plugins/MovieScene/MovieRenderPipeline/Source/MovieRenderPipelineEditor/Private/Customizations/Graph/MovieGraphLightModifierCustomization.cpp
+++ b/Engine/Plugins/MovieScene/MovieRenderPipeline/Source/MovieRenderPipelineEditor/Private/Customizations/Graph/MovieGraphLightModifierCustomization.cpp
@@ -439,7 +439,7 @@ TSharedRef<SWidget> FMovieGraphLightModifierCustomization::GetCustomMenuContents
 		SNew(STextBlock)
 		.Margin(FMargin(10, 5, 10, 5))
 		.ColorAndOpacity(FSlateColor::UseSubduedForeground())
-		.Text(LOCTEXT("CustomMenuHelpText", "Add an additional light actor property to this modifier."));
+		.Text(LOCTEXT("CustomMenuHelpText", "Add an additional light component property to this modifier."));
 
 	constexpr bool bNoIndent = true;
 	CustomMenu.AddWidget(CustomMenuHelpTextWidget, FText::GetEmpty(), bNoIndent);
@@ -449,7 +449,7 @@ TSharedRef<SWidget> FMovieGraphLightModifierCustomization::GetCustomMenuContents
 	TArray<UClass*> LightClasses;
 	PropertiesToDisplay.GetKeys(LightClasses);
 
-	CustomMenu.BeginSection(NAME_None, LOCTEXT("LightActorTypesSectionHeader", "Light Actor Types"));
+	CustomMenu.BeginSection(NAME_None, LOCTEXT("LightActorTypesSectionHeader", "Light Component Types"));
 
 	for (UClass* LightClass : LightClasses)
 	{
