MoviePipeline: Fixed issue where the Current Viewport mode in Quick Render would not reflect the viewport camera properly if it was transformed within Sequencer.

#Synced-CL 50204302

diff --git a/Engine/Plugins/MovieScene/MovieRenderPipeline/Source/MovieRenderPipelineEditor/MovieRenderPipelineEditor.Build.cs b/Engine/Plugins/MovieScene/MovieRenderPipeline/Source/MovieRenderPipelineEditor/MovieRenderPipelineEditor.Build.cs
--- a/Engine/Plugins/MovieScene/MovieRenderPipeline/Source/MovieRenderPipelineEditor/MovieRenderPipelineEditor.Build.cs
+++ b/Engine/Plugins/MovieScene/MovieRenderPipeline/Source/MovieRenderPipelineEditor/MovieRenderPipelineEditor.Build.cs
@@ -77,6 +77,7 @@ public class MovieRenderPipelineEditor : ModuleRules
 				"UEOpenExr",	// out of the ImageSequenceOutputNode header at some point so these dependencies can be removed.
 				"Projects",		// For IPluginManager
 				"NamingTokensUI",
+				"CinematicCamera",
 			}
         );
 
diff --git a/Engine/Plugins/MovieScene/MovieRenderPipeline/Source/MovieRenderPipelineEditor/Private/Graph/MovieGraphQuickRender.cpp b/Engine/Plugins/MovieScene/MovieRenderPipeline/Source/MovieRenderPipelineEditor/Private/Graph/MovieGraphQuickRender.cpp
--- a/Engine/Plugins/MovieScene/MovieRenderPipeline/Source/MovieRenderPipelineEditor/Private/Graph/MovieGraphQuickRender.cpp
+++ b/Engine/Plugins/MovieScene/MovieRenderPipeline/Source/MovieRenderPipelineEditor/Private/Graph/MovieGraphQuickRender.cpp
@@ -3,6 +3,7 @@
 #include "Graph/MovieGraphQuickRender.h"
 
 #include "Camera/CameraActor.h"
+#include "CineCameraActor.h"
 #include "Editor.h"
 #include "Engine/LevelStreaming.h"
 #include "EngineUtils.h"
@@ -916,10 +917,21 @@ void UMovieGraphQuickRenderSubsystem::CachePrePieData(UWorld* InEditorWorld)
 		const FLevelViewportActorLock& ActorLock = ViewportClient->IsLockedToCinematic() ? ViewportClient->GetCinematicActorLock() : ViewportClient->GetActorLock();
 		if (ActorLock.HasValidLockedActor())
 		{
-			CachedPrePieData.ViewportActorLock = ActorLock.LockedActor;
+			// Duplicate the pilot cam because it may be spawned by Sequencer. If that's the case, it'll be null by the time PIE starts up again.
+			CachedPrePieData.ViewportActorLock.Reset(DuplicateObject<AActor>(ActorLock.LockedActor.Get(), this));
+
+			// Cache the pilot cam's current world position. Needs to be cached out now because Sequencer will potentially change it
+			// once PIE starts up.
+			CachedPrePieData.bUseViewportActorLockWorldTransform = true;
+			CachedPrePieData.ViewportActorLockWorldTransform = GetWorldSpaceCameraTransform(ActorLock.LockedActor.Get());
 		}
 		
-		CachedPrePieData.ViewportActorLockCameraComponent = ViewportClient->GetCameraComponentForView();
+		// Duplicate the pilot cam component if there is one. As soon as PIE starts, the component's values will change to reflect the PIE world, so
+		// we duplicate it in order to be able to reference the pre-PIE property values.
+		if (const UCameraComponent* PilotCamComponent = ViewportClient->GetCameraComponentForView())
+		{
+			CachedPrePieData.ViewportActorLockCameraComponent.Reset(DuplicateObject<UCameraComponent>(PilotCamComponent, this));
+		}
 	}
 }
 
@@ -975,7 +987,8 @@ void UMovieGraphQuickRenderSubsystem::RestorePreRenderState()
 	CachedPrePieData.TemporarilyNonTransientComponents.Empty();
 	CachedPrePieData.TemporarilyNonTransientActors.Empty();
 	CachedPrePieData.ViewportActorLock.Reset();
-	CachedPrePieData.ViewportActorLockCameraComponent = nullptr;
+	CachedPrePieData.ViewportActorLockCameraComponent.Reset();
+	CachedPrePieData.bUseViewportActorLockWorldTransform = false;
 }
 
 ULevelSequence* UMovieGraphQuickRenderSubsystem::SetUpRenderingLevelSequence() const
@@ -1261,19 +1274,42 @@ void UMovieGraphQuickRenderSubsystem::PerformPreRenderSetup(UWorld* InEditorWorl
 		// If the viewport is piloting a camera (and not another type of non-camera actor), we should just duplicate it. We can't just copy over
 		// the basic transform + FOV, etc because this may be a cine camera with many properties set on it (like auto-exposure).
 		// Note that this supports "cinematic locks" (like cameras piloted via Sequencer) and normal piloted cameras.
-		if (CachedPrePieData.ViewportActorLock.IsValid() && IsValid(CachedPrePieData.ViewportActorLockCameraComponent))
+		if (CachedPrePieData.ViewportActorLock.IsValid() && CachedPrePieData.ViewportActorLockCameraComponent.IsValid())
 		{
-			AActor* ViewportPilotActor = CachedPrePieData.ViewportActorLock.Get();
-			if (ViewportPilotActor->IsA<ACameraActor>())
+			if (ACameraActor* ViewportPilotActor = Cast<ACameraActor>(CachedPrePieData.ViewportActorLock.Get()))
 			{
 				// Use the current pilot cam as the template
 				CameraSpawnParams.Template = ViewportPilotActor;
 
-				ACameraActor* SpawnedCamera =
-					Cast<ACameraActor>(PieWorld->SpawnActor(ViewportPilotActor->GetClass(), nullptr, CameraSpawnParams));
+				ACameraActor* SpawnedCamera = Cast<ACameraActor>(PieWorld->SpawnActor(ViewportPilotActor->GetClass(), nullptr, CameraSpawnParams));
 				if (SpawnedCamera)
 				{
-					SpawnedCamera->SetActorTransform(GetWorldSpaceCameraTransform(ViewportPilotActor));
+					// Copy over cine camera components to the spawned camera if the viewport was piloted by a cine cam
+					if (ACineCameraActor* SpawnedCineCam = Cast<ACineCameraActor>(SpawnedCamera))
+					{
+						UCineCameraComponent* CachedCineCamComponent = Cast<UCineCameraComponent>(CachedPrePieData.ViewportActorLockCameraComponent.Get());
+						UCineCameraComponent* DestCineCamComponent = SpawnedCineCam->GetCineCameraComponent();
+
+						// Copy over the cine cam properties to the destination component. Unfortunately UEngine::CopyPropertiesForUnrelatedObjects()
+						// does not appear to work here for some reason, but luckily the component has its settings grouped together for easy copying.
+						if (CachedCineCamComponent && DestCineCamComponent)
+						{
+							DestCineCamComponent->SetFilmback(CachedCineCamComponent->Filmback);
+							DestCineCamComponent->SetLensSettings(CachedCineCamComponent->LensSettings);
+							DestCineCamComponent->SetFocusSettings(CachedCineCamComponent->FocusSettings);
+							DestCineCamComponent->SetCropSettings(CachedCineCamComponent->CropSettings);
+							DestCineCamComponent->SetCurrentFocalLength(CachedCineCamComponent->CurrentFocalLength);
+							DestCineCamComponent->SetCurrentAperture(CachedCineCamComponent->CurrentAperture);
+							DestCineCamComponent->CurrentFocusDistance = CachedCineCamComponent->CurrentFocusDistance;
+							DestCineCamComponent->ExposureMethod = CachedCineCamComponent->ExposureMethod;
+						}
+					}
+
+					// Ensure the spawned camera has the correct world-space transform
+					const FTransform CameraTransform = CachedPrePieData.bUseViewportActorLockWorldTransform
+						? CachedPrePieData.ViewportActorLockWorldTransform
+						: GetWorldSpaceCameraTransform(ViewportPilotActor);
+					SpawnedCamera->SetActorTransform(CameraTransform);
 				}
 
 				return SpawnedCamera;
@@ -1297,13 +1333,10 @@ void UMovieGraphQuickRenderSubsystem::PerformPreRenderSetup(UWorld* InEditorWorl
 			
 			// If the viewport is being piloted, then we should use the properties of that to render from instead of the viewport's
 			// non-piloted camera properties. Note that this only covers non-camera pilot actors; camera pilot actors are covered above.
-			if (CachedPrePieData.ViewportActorLock.IsValid())
+			if (CachedPrePieData.bUseViewportActorLockWorldTransform)
 			{
-				const TStrongObjectPtr<AActor> LockedActorPin = CachedPrePieData.ViewportActorLock.Pin();
-
-				const FTransform WorldTransform = GetWorldSpaceCameraTransform(LockedActorPin.Get());
-				CameraPosition = WorldTransform.GetTranslation();
-				CameraRotation = WorldTransform.GetRotation().Rotator();
+				CameraPosition = CachedPrePieData.ViewportActorLockWorldTransform.GetTranslation();
+				CameraRotation = CachedPrePieData.ViewportActorLockWorldTransform.GetRotation().Rotator();
 			}
 			else
 			{
diff --git a/Engine/Plugins/MovieScene/MovieRenderPipeline/Source/MovieRenderPipelineEditor/Public/Graph/MovieGraphQuickRender.h b/Engine/Plugins/MovieScene/MovieRenderPipeline/Source/MovieRenderPipelineEditor/Public/Graph/MovieGraphQuickRender.h
--- a/Engine/Plugins/MovieScene/MovieRenderPipeline/Source/MovieRenderPipelineEditor/Public/Graph/MovieGraphQuickRender.h
+++ b/Engine/Plugins/MovieScene/MovieRenderPipeline/Source/MovieRenderPipelineEditor/Public/Graph/MovieGraphQuickRender.h
@@ -69,10 +69,19 @@ private:
 		TArray<TWeakObjectPtr<UActorComponent>> TemporarilyNonTransientComponents;
 
 		/** The actor that the viewport is locked to (also referred to as the pilot camera). Nullptr if no actor lock is active. */
-		TWeakObjectPtr<AActor> ViewportActorLock = nullptr;
+		TStrongObjectPtr<AActor> ViewportActorLock = nullptr;
 
-		/** If the viewport is locked to an actor (see ViewportLockActor), this is the specific camera component that it is locked to. May be nullptr if the viewport is locked to a non-camera actor (like a light). */
-		UCameraComponent* ViewportActorLockCameraComponent = nullptr;
+		/** Cached world transform of the viewport actor lock (pilot camera). */
+		FTransform ViewportActorLockWorldTransform = FTransform::Identity;
+
+		/** Whether the viewport actor lock world transform should actually be used. */
+		bool bUseViewportActorLockWorldTransform = false;
+
+		/**
+		 * If the viewport is locked to an actor (see ViewportActorLock), this is the specific camera component that it is locked to. May be
+		 * nullptr if the viewport is locked to a non-camera actor (like a light).
+		 */
+		TStrongObjectPtr<UCameraComponent> ViewportActorLockCameraComponent = nullptr;
 	};
 	
 	/**
