MoviePipeline: Fixed potential crash on exit with null pointer to output merger.

#Synced-CL 50287290

diff --git a/Engine/Plugins/MovieScene/MovieRenderPipeline/Source/MovieRenderPipelineCore/Private/Graph/MovieGraphPipeline.cpp b/Engine/Plugins/MovieScene/MovieRenderPipeline/Source/MovieRenderPipelineCore/Private/Graph/MovieGraphPipeline.cpp
--- a/Engine/Plugins/MovieScene/MovieRenderPipeline/Source/MovieRenderPipelineCore/Private/Graph/MovieGraphPipeline.cpp
+++ b/Engine/Plugins/MovieScene/MovieRenderPipeline/Source/MovieRenderPipelineCore/Private/Graph/MovieGraphPipeline.cpp
@@ -1733,7 +1733,10 @@ void UMovieGraphPipeline::BeginDestroy()
 {
 	// Ask the output merger to clean up. It can hold strong obj pointers, and those should not stick around after the pipeline shuts down. Generally
 	// this shouldn't be needed, but could be necessary if pending frames were not flushed for some reason.
-	GetOutputMerger()->AbandonOutstandingWork();
+	if (OutputMerger.IsValid())
+	{
+		OutputMerger->AbandonOutstandingWork();
+	}
 	
 	Super::BeginDestroy();
 }
