Sequencer: Fix MRQ non realtime render to render silence instead of disabling

#Synced-CL 50528440

diff --git a/Engine/Plugins/MovieScene/MovieRenderPipeline/Source/MovieRenderPipelineCore/Private/MoviePipelineGameOverrideSetting.cpp b/Engine/Plugins/MovieScene/MovieRenderPipeline/Source/MovieRenderPipelineCore/Private/MoviePipelineGameOverrideSetting.cpp
--- a/Engine/Plugins/MovieScene/MovieRenderPipeline/Source/MovieRenderPipelineCore/Private/MoviePipelineGameOverrideSetting.cpp
+++ b/Engine/Plugins/MovieScene/MovieRenderPipeline/Source/MovieRenderPipelineCore/Private/MoviePipelineGameOverrideSetting.cpp
@@ -147,6 +147,9 @@ void UMoviePipelineGameOverrideSetting::ApplyCVarSettings(const bool bOverrideVa
 
 	MOVIEPIPELINE_STORE_AND_OVERRIDE_CVAR_INT(PreviousNeverMuteNonRealtimeAudio, TEXT("au.NeverMuteNonRealtimeAudioDevices"), 1, bOverrideValues);
 
+	// Ensures that submixes will always render silence instead of disabling to maintain timing in the sequence
+	MOVIEPIPELINE_STORE_AND_OVERRIDE_CVAR_BOOL(PreviousNeverDisableSubmixes, TEXT("au.NeverDisableSubmixes"), 1, bOverrideValues);
+
 	// To make sure that the skylight is always valid and consistent accross capture sessions, we enforce a full capture each frame, accepting a small GPU cost.
 	MOVIEPIPELINE_STORE_AND_OVERRIDE_CVAR_INT(PreviousSkyLightRealTimeReflectionCaptureTimeSlice, TEXT("r.SkyLight.RealTimeReflectionCapture.TimeSlice"), 0, bOverrideValues);
 
@@ -289,6 +292,7 @@ void UMoviePipelineGameOverrideSetting::BuildNewProcessCommandLineArgsImpl(TArra
 	}
 
 	InOutDeviceProfileCvars.AddUnique(FString::Printf(TEXT("au.NeverMuteNonRealtimeAudioDevices=%d"), 1));
+	InOutDeviceProfileCvars.AddUnique(FString::Printf(TEXT("au.NeverDisableSubmixes=%d"), 1));
 	InOutDeviceProfileCvars.AddUnique(FString::Printf(TEXT("r.SkyLight.RealTimeReflectionCapture.TimeSlice=%d"), 0));
 	InOutDeviceProfileCvars.AddUnique(FString::Printf(TEXT("r.VolumetricRenderTarget=%d"), 1));
 	InOutDeviceProfileCvars.AddUnique(FString::Printf(TEXT("r.VolumetricRenderTarget.Mode=%d"), 3));
diff --git a/Engine/Plugins/MovieScene/MovieRenderPipeline/Source/MovieRenderPipelineCore/Public/MoviePipelineGameOverrideSetting.h b/Engine/Plugins/MovieScene/MovieRenderPipeline/Source/MovieRenderPipelineCore/Public/MoviePipelineGameOverrideSetting.h
--- a/Engine/Plugins/MovieScene/MovieRenderPipeline/Source/MovieRenderPipelineCore/Public/MoviePipelineGameOverrideSetting.h
+++ b/Engine/Plugins/MovieScene/MovieRenderPipeline/Source/MovieRenderPipelineCore/Public/MoviePipelineGameOverrideSetting.h
@@ -172,6 +172,7 @@ private:
 	int32 PreviousFoliageDitheredLOD;
 	int32 PreviousFoliageForceLOD;
 	int32 PreviousNeverMuteNonRealtimeAudio;
+	bool PreviousNeverDisableSubmixes;
 	int32 PreviousSkyLightRealTimeReflectionCaptureTimeSlice;
 	int32 PreviousVolumetricRenderTarget;
 	int32 PreviousVolumetricRenderTargetMode;
diff --git a/Engine/Source/Runtime/AudioMixer/Private/AudioMixerSubmix.cpp b/Engine/Source/Runtime/AudioMixer/Private/AudioMixerSubmix.cpp
--- a/Engine/Source/Runtime/AudioMixer/Private/AudioMixerSubmix.cpp
+++ b/Engine/Source/Runtime/AudioMixer/Private/AudioMixerSubmix.cpp
@@ -46,6 +46,16 @@ FAutoConsoleVariableRef CVarLogSubmixEnablement(
 	TEXT("1: Submix enablement logging is on. 0: Submix enablement/disablement logging is off."),
 	ECVF_Default);
 
+static bool bNeverDisableSubmixes = false;
+static FAutoConsoleVariableRef CVarNeverDisableSubmixes(
+	TEXT("au.NeverDisableSubmixes"),
+	bNeverDisableSubmixes,
+	TEXT("Forces submixes to stay enabled and ignore the bAutoDisable value for situations where rendering silence is required instead of no audio rendering\n")
+	TEXT("0: Submixes use auto disable as normal\n")
+	TEXT("1: Submixes always render silence instead of disabling\n"),
+	ECVF_Default
+);
+
 // Define profiling categories for submixes. 
 DEFINE_STAT(STAT_AudioMixerEndpointSubmixes);
 DEFINE_STAT(STAT_AudioMixerSubmixes);
@@ -1339,7 +1349,7 @@ namespace Audio
 		float* BufferPtr = InputBuffer.GetData();
 
 		// As an optimization, early out if we're set to auto-disable and we're not rendering audio
-		if (bAutoDisable && !IsRenderingAudio())
+		if (bAutoDisable && !IsRenderingAudio() && !bNeverDisableSubmixes)
 		{
 			if (!bIsCurrentlyDisabled)
 			{
@@ -1694,7 +1704,7 @@ namespace Audio
 		const bool bShouldTraceHasActivityMsg = UE_TRACE_CHANNELEXPR_IS_ENABLED(AudioChannel) && OwningSubmixObject.IsValid();
 #endif // UE_AUDIO_PROFILERTRACE_ENABLED
 
-		if (bAutoDisable)
+		if (bAutoDisable && !bNeverDisableSubmixes)
 		{
 			bIsNowSilent = FMixerSubmixPrivate::IsAudioBufferSilent(OutAudioBuffer);
 
