MoviePipeline: We default-increase the engine's combined LUT size during renders.

* Minimize potential quantization errors, which have been reported with subsample accumulation & elsewhere.
* LUT texture increases from the default 32^3 to 64^3 (~0.3 MB to ~2MB with PF_FloatRGBA)
* Validated with MRQTest suite.

#Synced-CL 50473050

diff --git a/Engine/Plugins/MovieScene/MovieRenderPipeline/Source/MovieRenderPipelineCore/Private/Graph/Nodes/MovieGraphGlobalGameOverrides.cpp b/Engine/Plugins/MovieScene/MovieRenderPipeline/Source/MovieRenderPipelineCore/Private/Graph/Nodes/MovieGraphGlobalGameOverrides.cpp
--- a/Engine/Plugins/MovieScene/MovieRenderPipeline/Source/MovieRenderPipelineCore/Private/Graph/Nodes/MovieGraphGlobalGameOverrides.cpp
+++ b/Engine/Plugins/MovieScene/MovieRenderPipeline/Source/MovieRenderPipelineCore/Private/Graph/Nodes/MovieGraphGlobalGameOverrides.cpp
@@ -109,6 +109,7 @@ void UMovieGraphGlobalGameOverridesNode::BuildNewProcessCommandLineArgsImpl(TArr
 	}
 
 	InOutDeviceProfileCvars.AddUnique(FString::Printf(TEXT("r.PostProcessing.PropagateAlpha=%d"), UE::MoviePipeline::GetAlphaOutputOverride() ? 1 : 0));
+	InOutDeviceProfileCvars.AddUnique(FString::Printf(TEXT("r.LUT.Size=%d"), 64));
 }
 
 void UMovieGraphGlobalGameOverridesNode::PostLoad()
@@ -193,8 +194,6 @@ void UMovieGraphGlobalGameOverridesNode::ApplySettings(const bool bOverrideValue
 		MOVIEPIPELINE_STORE_AND_OVERRIDE_CVAR_INT(NodeCDO->PreviousLumenSurfaceCacheRecaptureEveryFrame, TEXT("r.LumenScene.SurfaceCache.RecaptureEveryFrame"), 1, bOverrideValues);
 	}
 
-	MOVIEPIPELINE_STORE_AND_OVERRIDE_CVAR_BOOL(NodeCDO->PreviousAlphaOutput, TEXT("r.PostProcessing.PropagateAlpha"), UE::MoviePipeline::GetAlphaOutputOverride(), bOverrideValues);
-
 	// Set cvars that allow MRQ to run correctly. These are not exposed as properties on the node, but the chance of a
 	// user needing to customize these (or be aware of them) is very low.
 	{
@@ -243,6 +242,11 @@ void UMovieGraphGlobalGameOverridesNode::ApplySettings(const bool bOverrideValue
 			MOVIEPIPELINE_STORE_AND_OVERRIDE_CVAR_INT(NodeCDO->PreviousNaniteVSMInvalidateOnLODDelta, TEXT("r.Nanite.VSMInvalidateOnLODDelta"), 1, bOverrideValues);
 		}
 
+		// Enable the alpha channel throughout the renderer
+		MOVIEPIPELINE_STORE_AND_OVERRIDE_CVAR_BOOL(NodeCDO->PreviousAlphaOutput, TEXT("r.PostProcessing.PropagateAlpha"), UE::MoviePipeline::GetAlphaOutputOverride(), bOverrideValues);
+
+		// Improve (combined) LUT precision to avoid quantization errors, especially with multi-sample accumulation
+		MOVIEPIPELINE_STORE_AND_OVERRIDE_CVAR_INT(NodeCDO->PreviousLUTSize, TEXT("r.LUT.Size"), 64, bOverrideValues);
 	}
 	
 	// Must come after the above cvars so that if one of those cvars is also specified by the scalability level, then we
diff --git a/Engine/Plugins/MovieScene/MovieRenderPipeline/Source/MovieRenderPipelineCore/Private/MoviePipelineGameOverrideSetting.cpp b/Engine/Plugins/MovieScene/MovieRenderPipeline/Source/MovieRenderPipelineCore/Private/MoviePipelineGameOverrideSetting.cpp
--- a/Engine/Plugins/MovieScene/MovieRenderPipeline/Source/MovieRenderPipelineCore/Private/MoviePipelineGameOverrideSetting.cpp
+++ b/Engine/Plugins/MovieScene/MovieRenderPipeline/Source/MovieRenderPipelineCore/Private/MoviePipelineGameOverrideSetting.cpp
@@ -131,8 +131,6 @@ void UMoviePipelineGameOverrideSetting::ApplyCVarSettings(const bool bOverrideVa
 	{
 		MOVIEPIPELINE_STORE_AND_OVERRIDE_CVAR_INT(PreviousStreamingManagerSyncState, TEXT("r.Streaming.SyncStatesWhenBlocking"), 1, bOverrideValues);
 	}
-
-	MOVIEPIPELINE_STORE_AND_OVERRIDE_CVAR_BOOL(PreviousAlphaOutput, TEXT("r.PostProcessing.PropagateAlpha"), UE::MoviePipeline::GetAlphaOutputOverride(), bOverrideValues);
 	
 #if WITH_EDITOR
 	// To make sure the GeometryCache streamer doesn't skip frames and doesn't pop up notification during rendering
@@ -184,6 +182,12 @@ void UMoviePipelineGameOverrideSetting::ApplyCVarSettings(const bool bOverrideVa
 	{
 		MOVIEPIPELINE_STORE_AND_OVERRIDE_CVAR_INT(PreviousNaniteVSMInvalidateOnLODDelta, TEXT("r.Nanite.VSMInvalidateOnLODDelta"), 1, bOverrideValues);
 	}
+	
+	// Enable the alpha channel throughout the renderer
+	MOVIEPIPELINE_STORE_AND_OVERRIDE_CVAR_BOOL(PreviousAlphaOutput, TEXT("r.PostProcessing.PropagateAlpha"), UE::MoviePipeline::GetAlphaOutputOverride(), bOverrideValues);
+
+	// Improve (combined) LUT precision to avoid quantization errors, especially with multi-sample accumulation
+	MOVIEPIPELINE_STORE_AND_OVERRIDE_CVAR_INT(PreviousLUTSize, TEXT("r.LUT.Size"), 64, bOverrideValues);
 
 	// Must come after the above cvars so that if one of those cvars is also specified by the Scalability level, then we restore to the value in the original scalability level
 	// not the value we cached in the Cinematic level (if applied).
@@ -309,4 +313,5 @@ void UMoviePipelineGameOverrideSetting::BuildNewProcessCommandLineArgsImpl(TArra
 	}
 
 	InOutDeviceProfileCvars.AddUnique(FString::Printf(TEXT("r.PostProcessing.PropagateAlpha=%d"), UE::MoviePipeline::GetAlphaOutputOverride() ? 1 : 0));
+	InOutDeviceProfileCvars.AddUnique(FString::Printf(TEXT("r.LUT.Size=%d"), 64));
 }
diff --git a/Engine/Plugins/MovieScene/MovieRenderPipeline/Source/MovieRenderPipelineCore/Public/Graph/Nodes/MovieGraphGlobalGameOverrides.h b/Engine/Plugins/MovieScene/MovieRenderPipeline/Source/MovieRenderPipelineCore/Public/Graph/Nodes/MovieGraphGlobalGameOverrides.h
--- a/Engine/Plugins/MovieScene/MovieRenderPipeline/Source/MovieRenderPipelineCore/Public/Graph/Nodes/MovieGraphGlobalGameOverrides.h
+++ b/Engine/Plugins/MovieScene/MovieRenderPipeline/Source/MovieRenderPipelineCore/Public/Graph/Nodes/MovieGraphGlobalGameOverrides.h
@@ -247,6 +247,7 @@ private:
 	int32 PreviousNaniteVSMInvalidateOnLODDelta;
 	int32 PreviousAllowSlateThrottling;
 	bool PreviousAlphaOutput;
+	int32 PreviousLUTSize;
 
 #if WITH_EDITOR
 	int32 PreviousGeoCacheStreamerShowNotification;
diff --git a/Engine/Plugins/MovieScene/MovieRenderPipeline/Source/MovieRenderPipelineCore/Public/MoviePipelineGameOverrideSetting.h b/Engine/Plugins/MovieScene/MovieRenderPipeline/Source/MovieRenderPipelineCore/Public/MoviePipelineGameOverrideSetting.h
--- a/Engine/Plugins/MovieScene/MovieRenderPipeline/Source/MovieRenderPipelineCore/Public/MoviePipelineGameOverrideSetting.h
+++ b/Engine/Plugins/MovieScene/MovieRenderPipeline/Source/MovieRenderPipelineCore/Public/MoviePipelineGameOverrideSetting.h
@@ -189,6 +189,7 @@ private:
 	int32 PreviousNaniteVSMInvalidateOnLODDelta;
 	int32 PreviousAllowSlateThrottling;
 	bool PreviousAlphaOutput;
+	int32 PreviousLUTSize;
 };
 
 #undef UE_API
