MoviePipeline: Fixed a number of problems with the Light Modifier node.
* Fixed a crash that could occur when copy/paste-ing or duplicating a Light Modifier node.
* Fixed a bug that caused the Light Channels property to potentially apply incorrectly.
* Fixed a bug that caused lighting properties to not revert correctly in some circumstances, leading to incorrect renders in multi-layer renders.
* Fixed a regression that caused "Custom" properties to not be listed in the UI within the correct group, or apply correctly to lights.

#rnx

#Synced-CL 50695613

diff --git a/Engine/Plugins/MovieScene/MovieRenderPipeline/Source/MovieRenderPipelineCore/Private/Graph/MovieGraphNode.cpp b/Engine/Plugins/MovieScene/MovieRenderPipeline/Source/MovieRenderPipelineCore/Private/Graph/MovieGraphNode.cpp
--- a/Engine/Plugins/MovieScene/MovieRenderPipeline/Source/MovieRenderPipelineCore/Private/Graph/MovieGraphNode.cpp
+++ b/Engine/Plugins/MovieScene/MovieRenderPipeline/Source/MovieRenderPipelineCore/Private/Graph/MovieGraphNode.cpp
@@ -100,11 +100,7 @@ namespace UE::MovieGraph::Private
 
 UMovieGraphNode::UMovieGraphNode()
 {
-	if (!HasAnyFlags(RF_ClassDefaultObject))
-	{
-		DynamicPropertiesView = NewObject<UMovieGraphMutableValueView>(this, TEXT("DynamicPropertiesView"));
-		DynamicPropertiesView->ViewProperties(&DynamicProperties);
-	}
+	
 }
 
 TArray<FMovieGraphPinProperties> UMovieGraphNode::GetExposedPinProperties() const
@@ -500,6 +496,14 @@ bool UMovieGraphNode::GetDynamicPropertyValue(const FName PropertyName, FString&
 
 UMovieGraphMutableValueView* UMovieGraphNode::GetDynamicPropertiesView() const
 {
+	// Lazily instantiate the view
+	if (!DynamicPropertiesView)
+	{
+		UMovieGraphNode* MutableThis = const_cast<UMovieGraphNode*>(this);
+		DynamicPropertiesView = NewObject<UMovieGraphMutableValueView>(MutableThis, TEXT("DynamicPropertiesView"));
+		DynamicPropertiesView->ViewProperties(&MutableThis->DynamicProperties);
+	}
+
 	return DynamicPropertiesView.Get();
 }
 
diff --git a/Engine/Plugins/MovieScene/MovieRenderPipeline/Source/MovieRenderPipelineCore/Private/Graph/Nodes/MovieGraphClassPropertyModifier.cpp b/Engine/Plugins/MovieScene/MovieRenderPipeline/Source/MovieRenderPipelineCore/Private/Graph/Nodes/MovieGraphClassPropertyModifier.cpp
--- a/Engine/Plugins/MovieScene/MovieRenderPipeline/Source/MovieRenderPipelineCore/Private/Graph/Nodes/MovieGraphClassPropertyModifier.cpp
+++ b/Engine/Plugins/MovieScene/MovieRenderPipeline/Source/MovieRenderPipelineCore/Private/Graph/Nodes/MovieGraphClassPropertyModifier.cpp
@@ -113,15 +113,17 @@ void UMovieGraphPropertyModifier::ApplyModifier(const UWorld* World)
 		// properly handle every type. This will do for now, but could potentially be optimized in the future if this is a bottleneck
 		// or causes issues.
 
+		const FProperty* TargetProperty = InPropertyReference.GetProperty();
+
 		// Cache the current value so it can be restored later
 		FString CurrentValue;
-		InPropertyReference.GetProperty()->ExportTextItem_InContainer(CurrentValue, InPropertyContainer, nullptr, nullptr, PPF_None);
+		TargetProperty->ExportTextItem_InContainer(CurrentValue, InPropertyContainer, nullptr, nullptr, PPF_None);
 		PreviousPropertyValuesView->SetValueSerializedString(InPropertyReference.AuxPropertyName, CurrentValue);
 		ModifiedProperties.Add({ InPropertyContainer, InPropertyReference });
 
 		// Set the property to its new value
 		const FString PropertyValue = UpdatedPropertyValuesView->GetValueSerializedString(InPropertyReference.AuxPropertyName);
-		InPropertyReference.GetProperty()->ImportText_InContainer(*PropertyValue, InPropertyContainer, InPropertyContainer, 0);
+		TargetProperty->ImportText_InContainer(*PropertyValue, InPropertyContainer, InPropertyContainer, 0);
 	};
 
 	TArray<TSubclassOf<AActor>> TargetActorClasses;
@@ -167,7 +169,11 @@ void UMovieGraphPropertyModifier::ApplyModifier(const UWorld* World)
 				if (PropertyReference.GetComponentClass())
 				{
 					PropertyContainer = MatchingActor->GetComponentByClass(PropertyReference.GetComponentClass());
-					DirtiedComponents.Add(Cast<UActorComponent>(PropertyContainer));
+					
+					if (PropertyContainer)
+					{
+						DirtiedComponents.Add(Cast<UActorComponent>(PropertyContainer));
+					}
 				}
 				else if (!bOnlyMatchComponent)
 				{
@@ -233,8 +239,13 @@ void UMovieGraphPropertyModifier::ApplyModifier(const UWorld* World)
 
 void UMovieGraphPropertyModifier::UndoModifier()
 {
-	for (const TPair<UObject*, FMovieGraphPropertyReference>& ModifiedPropertyPair : ModifiedProperties)
+	// Undo property modifications in reverse. There's a possibility that a property could be modified multiple times on the same object;
+	// to ensure that the original property value is correctly restored, the first modification to it must be restored *last*. This is due to how
+	// property values are set + cached at the same time (see CacheAndSetPropertyValue()).
+	for (int32 Index = ModifiedProperties.Num() - 1; Index >= 0; --Index)
 	{
+		const TPair<UObject*, FMovieGraphPropertyReference>& ModifiedPropertyPair = ModifiedProperties[Index];
+
 		UObject* ModifiedObject = ModifiedPropertyPair.Key;
 		const FMovieGraphPropertyReference& ModifiedPropertyReference = ModifiedPropertyPair.Value;
 		const FProperty* ModifiedProperty = ModifiedPropertyPair.Value.GetProperty();
diff --git a/Engine/Plugins/MovieScene/MovieRenderPipeline/Source/MovieRenderPipelineCore/Private/Graph/Nodes/MovieGraphLightModifierNode.cpp b/Engine/Plugins/MovieScene/MovieRenderPipeline/Source/MovieRenderPipelineCore/Private/Graph/Nodes/MovieGraphLightModifierNode.cpp
--- a/Engine/Plugins/MovieScene/MovieRenderPipeline/Source/MovieRenderPipelineCore/Private/Graph/Nodes/MovieGraphLightModifierNode.cpp
+++ b/Engine/Plugins/MovieScene/MovieRenderPipeline/Source/MovieRenderPipelineCore/Private/Graph/Nodes/MovieGraphLightModifierNode.cpp
@@ -106,7 +106,7 @@ TArray<FPropertyBagPropertyDesc> UMovieGraphLightModifierNode::GetDynamicPropert
 		}
 
 		// Enable categories so that they're grouped by light type. The category also needs to be changed to the class associated with the property.
-		if (const TSubclassOf<AActor> PropertyClass = CustomPropertyRef.GetPropertyClass())
+		if (const TSubclassOf<UActorComponent> PropertyClass = CustomPropertyRef.GetComponentClass())
 		{
 			const FString LightingClassDisplayName = PropertyClass->GetDisplayNameText().ToString();
 
@@ -248,10 +248,30 @@ TArray<UMovieGraphModifierBase*> UMovieGraphLightModifierNode::GetAllModifiers()
 			const FProperty* NodeProperty = FindFProperty<FProperty>(UMovieGraphLightModifierNode::StaticClass(), InNodePropertyName);
 			if (ensureMsgf(NodeProperty, TEXT("Could not find property named [%s] on the Light Modifier node."), *InNodePropertyName.ToString()))
 			{
+				const FName ModifierPropertyName = Modifier->GetPropertyNameForSettingCustomValue(NewProperty);
+				
 				FString PropertyValue;
 				NodeProperty->ExportTextItem_InContainer(PropertyValue, this, nullptr, nullptr, PPF_None);
 
-				ModifierValues->SetValueSerializedString(Modifier->GetPropertyNameForSettingCustomValue(NewProperty), PropertyValue);
+				// Zero out the struct's memory (within the modifier's property bag). When calling ExportTextItem_InContainer() above, the exported/
+				// serialized text provides a *sparse* representation of what was changed on the struct, not a representation of all properties. Thus
+				// if we want the modifier's representation of the struct to represent the value set on the node, the struct first needs to be zero'd
+				// out, THEN the sparse representation can be applied.
+				if (const FStructProperty* StructTargetProperty = CastField<const FStructProperty>(NodeProperty))
+				{
+					const TObjectPtr<UScriptStruct> TargetScriptStruct = StructTargetProperty->Struct;
+
+					FStructView StructView;
+					if (ModifierValues->GetValueStruct(ModifierPropertyName, StructView, TargetScriptStruct.Get()))
+					{
+						void* StructData = StructView.GetMemory();
+						
+						TargetScriptStruct->DestroyStruct(StructData);
+						FMemory::Memzero(StructData, TargetScriptStruct->GetStructureSize());
+					}
+				}
+
+				ModifierValues->SetValueSerializedString(ModifierPropertyName, PropertyValue);
 			}
 		});
 	};
@@ -390,10 +410,14 @@ FString UMovieGraphLightModifierNode::GetNodeInstanceName() const
 	return ModifierName;
 }
 
-bool UMovieGraphLightModifierNode::AddCustomLightProperty(UClass* InLightClass, const FProperty* InLightingProperty)
+bool UMovieGraphLightModifierNode::AddCustomLightProperty(UClass* InLightComponentClass, const FProperty* InLightProperty)
 {
-	const FName PropertyName = GetLightingPropertyName(InLightClass, InLightingProperty);
+	if (!InLightComponentClass || !InLightProperty)
+	{
+		return false;
+	}
 
+	const FName PropertyName = GetLightingPropertyName(InLightComponentClass, InLightProperty);
 	const UMovieGraphMutableValueView* LightingValues = GetDynamicPropertiesView();
 
 	// Skip adding if the property already exists in the property bag.
@@ -408,11 +432,12 @@ bool UMovieGraphLightModifierNode::AddCustomLightProperty(UClass* InLightClass,
 
 	// Add supplemental property information for this property, which will importantly track the origin component.
 	// This has to be done via sidecar data, rather than via metadata on the property in the property bag, because runtime builds do not contain
-	// property metadata.
+	// property metadata. Note that the component class is tracked here instead of the actor class, since we want to potentially modify light
+	// components outside of the typical light actor classes.
 	constexpr UClass* LightClass = nullptr;
 	FMovieGraphPropertyReference& NewCustomProperty = CustomPropertyInfo.Emplace(
 		PropertyName,
-		FMovieGraphPropertyReference(LightClass, InLightingProperty->GetOwnerClass(), InLightingProperty->GetFName(), const_cast<FProperty*>(InLightingProperty)));
+		FMovieGraphPropertyReference(LightClass, InLightComponentClass, InLightProperty->GetFName(), const_cast<FProperty*>(InLightProperty)));
 	NewCustomProperty.AuxPropertyName = PropertyName;
 
 	// Update DynamicProperties with the properties referenced in CustomPropertyInfo. Although we could update DynamicProperties right here, there are
@@ -422,20 +447,20 @@ bool UMovieGraphLightModifierNode::AddCustomLightProperty(UClass* InLightClass,
 	return true;
 }
 
-bool UMovieGraphLightModifierNode::AddCustomLightProperty(UClass* InLightClass, const FName& InPropertyName)
+bool UMovieGraphLightModifierNode::AddCustomLightProperty(UClass* InLightComponentClass, const FName& InPropertyName)
 {
-	return AddCustomLightProperty(InLightClass, FindFProperty<FProperty>(InLightClass, InPropertyName));
+	return AddCustomLightProperty(InLightComponentClass, FindFProperty<FProperty>(InLightComponentClass, InPropertyName));
 }
 
-bool UMovieGraphLightModifierNode::RemoveCustomLightProperty(const UClass* InLightClass, const FProperty* InLightingProperty)
+bool UMovieGraphLightModifierNode::RemoveCustomLightProperty(const UClass* InLightComponentClass, const FProperty* InLightProperty)
 {
-	if (InLightClass && InLightingProperty)
+	if (InLightComponentClass && InLightProperty)
 	{
 #if WITH_EDITOR
 		Modify();
 #endif
 
-		const FName LightingPropertyName = GetLightingPropertyName(InLightClass, InLightingProperty);
+		const FName LightingPropertyName = GetLightingPropertyName(InLightComponentClass, InLightProperty);
 
 		const bool bSuccess = CustomPropertyInfo.Remove(LightingPropertyName) > 0;
 
@@ -448,27 +473,32 @@ bool UMovieGraphLightModifierNode::RemoveCustomLightProperty(const UClass* InLig
 	return false;
 }
 
-bool UMovieGraphLightModifierNode::RemoveCustomLightProperty(const UClass* InLightClass, const FName& InPropertyName)
+bool UMovieGraphLightModifierNode::RemoveCustomLightProperty(const UClass* InLightComponentClass, const FName& InPropertyName)
 {
-	return RemoveCustomLightProperty(InLightClass, FindFProperty<FProperty>(InLightClass, InPropertyName));
+	return RemoveCustomLightProperty(InLightComponentClass, FindFProperty<FProperty>(InLightComponentClass, InPropertyName));
 }
 
-bool UMovieGraphLightModifierNode::HasCustomLightProperty(const UClass* InLightClass, const FProperty* InLightingProperty) const
+bool UMovieGraphLightModifierNode::HasCustomLightProperty(const UClass* InLightComponentClass, const FProperty* InLightProperty) const
 {
+	if (!InLightComponentClass || !InLightProperty)
+	{
+		return false;
+	}
+
 	const UMovieGraphMutableValueView* LightingValues = GetDynamicPropertiesView();
 	
-	return LightingValues->HasProperty(GetLightingPropertyName(InLightClass, InLightingProperty));
+	return LightingValues->HasProperty(GetLightingPropertyName(InLightComponentClass, InLightProperty));
 }
 
-bool UMovieGraphLightModifierNode::HasCustomLightProperty(const UClass* InLightClass, const FName& InPropertyName) const
+bool UMovieGraphLightModifierNode::HasCustomLightProperty(const UClass* InLightComponentClass, const FName& InPropertyName) const
 {
-	return HasCustomLightProperty(InLightClass, FindFProperty<FProperty>(InLightClass, InPropertyName));
+	return HasCustomLightProperty(InLightComponentClass, FindFProperty<FProperty>(InLightComponentClass, InPropertyName));
 }
 
-bool UMovieGraphLightModifierNode::IsCustomLightPropertyOverridden(const UClass* InLightClass, const FName& InPropertyName) const
+bool UMovieGraphLightModifierNode::IsCustomLightPropertyOverridden(const UClass* InLightComponentClass, const FName& InPropertyName) const
 {
 	const UMovieGraphMutableValueView* LightingValues = GetDynamicPropertiesView();
-	const FName MainPropertyName = GetLightingPropertyName(InLightClass, FindFProperty<FProperty>(InLightClass, InPropertyName));
+	const FName MainPropertyName = GetLightingPropertyName(InLightComponentClass, FindFProperty<FProperty>(InLightComponentClass, InPropertyName));
 
 	bool bIsOverridden = false;
 	LightingValues->GetValueBool(GetLightingOverridePropertyName(MainPropertyName), bIsOverridden);
@@ -476,21 +506,21 @@ bool UMovieGraphLightModifierNode::IsCustomLightPropertyOverridden(const UClass*
 	return bIsOverridden;
 }
 
-bool UMovieGraphLightModifierNode::UpdateCustomLightPropertyOverrideState(const UClass* InLightClass, const FName& InPropertyName, const bool bIsOverridden)
+bool UMovieGraphLightModifierNode::UpdateCustomLightPropertyOverrideState(const UClass* InLightComponentClass, const FName& InPropertyName, const bool bIsOverridden)
 {
 	UMovieGraphMutableValueView* LightingValues = GetDynamicPropertiesView();
-	const FName MainPropertyName = GetLightingPropertyName(InLightClass, FindFProperty<FProperty>(InLightClass, InPropertyName));
+	const FName MainPropertyName = GetLightingPropertyName(InLightComponentClass, FindFProperty<FProperty>(InLightComponentClass, InPropertyName));
 
 	return LightingValues->SetValueBool(GetLightingOverridePropertyName(MainPropertyName), bIsOverridden);
 }
 
-FName UMovieGraphLightModifierNode::GetLightingPropertyName(const UClass* InLightClass, const FProperty* InLightingProperty)
+FName UMovieGraphLightModifierNode::GetLightingPropertyName(const UClass* InLightComponentClass, const FProperty* InLightProperty)
 {
 	static const TCHAR* LightingPropertyNameFormat = TEXT("{0}_{1}");
 
-	if (InLightClass && InLightingProperty)
+	if (InLightComponentClass && InLightProperty)
 	{
-		return FName(FString::Format(LightingPropertyNameFormat, { InLightClass->GetName(), InLightingProperty->GetName() }));
+		return FName(FString::Format(LightingPropertyNameFormat, { InLightComponentClass->GetName(), InLightProperty->GetName() }));
 	}
 
 	return FName();
@@ -508,7 +538,7 @@ int32 UMovieGraphLightModifierNode::GetNumCustomLightProperties() const
 	const UMovieGraphMutableValueView* LightingValues = GetDynamicPropertiesView();
 
 	// The divide by two is used to account for the bOverride_* properties; those should not be counted.
-	return LightingValues->GetNumProperties() / 2;
+	return LightingValues ? (LightingValues->GetNumProperties() / 2) : 0;
 }
 
 #undef LOCTEXT_NAMESPACE // "MovieGraph"
\ No newline at end of file
diff --git a/Engine/Plugins/MovieScene/MovieRenderPipeline/Source/MovieRenderPipelineCore/Public/Graph/MovieGraphNode.h b/Engine/Plugins/MovieScene/MovieRenderPipeline/Source/MovieRenderPipelineCore/Public/Graph/MovieGraphNode.h
--- a/Engine/Plugins/MovieScene/MovieRenderPipeline/Source/MovieRenderPipelineCore/Public/Graph/MovieGraphNode.h
+++ b/Engine/Plugins/MovieScene/MovieRenderPipeline/Source/MovieRenderPipelineCore/Public/Graph/MovieGraphNode.h
@@ -443,7 +443,7 @@ protected:
 
 	/** A property view which points to the dynamic properties on this node. */
 	UPROPERTY(Transient)
-	TObjectPtr<UMovieGraphMutableValueView> DynamicPropertiesView = nullptr;
+	mutable TObjectPtr<UMovieGraphMutableValueView> DynamicPropertiesView = nullptr;
 
 	/** Tracks which properties have been exposed on the node as inputs. */
 	UPROPERTY()
diff --git a/Engine/Plugins/MovieScene/MovieRenderPipeline/Source/MovieRenderPipelineCore/Public/Graph/Nodes/MovieGraphLightModifierNode.h b/Engine/Plugins/MovieScene/MovieRenderPipeline/Source/MovieRenderPipelineCore/Public/Graph/Nodes/MovieGraphLightModifierNode.h
--- a/Engine/Plugins/MovieScene/MovieRenderPipeline/Source/MovieRenderPipelineCore/Public/Graph/Nodes/MovieGraphLightModifierNode.h
+++ b/Engine/Plugins/MovieScene/MovieRenderPipeline/Source/MovieRenderPipelineCore/Public/Graph/Nodes/MovieGraphLightModifierNode.h
@@ -27,10 +27,10 @@ enum class EMovieGraphLightModifierIntensityMethod : uint8
 };
 
 /** 
- * A node which modifies properties on light actors within the world.
+ * A node which modifies properties on light components within the world.
  *
  * There are several common light properties (eg, Light Color) that can be modified directly. For less common light properties, they
- * can be modified by adding them as "Custom" properties. The common properties apply to all light types; custom properties apply only
+ * can be modified by adding them as "Custom" properties. The common properties apply to all standard light types; custom properties apply only
  * to a specific light type.
  */
 UCLASS(MinimalAPI)
@@ -79,7 +79,7 @@ public:
 	 * Adding a custom property does not immediately make it take effect. Its override state needs to be set to "true" via
 	 * UpdateCustomLightPropertyOverrideState().
 	 */
-	UE_API bool AddCustomLightProperty(UClass* InLightClass, const FProperty* InLightingProperty);
+	UE_API bool AddCustomLightProperty(UClass* InLightComponentClass, const FProperty* InLightProperty);
 
 	/**
 	 * Adds the property from the specified class to the "Custom" lighting properties. Does nothing if the Custom properties already contains it.
@@ -89,41 +89,41 @@ public:
 	 * UpdateCustomLightPropertyOverrideState().
 	 */
 	UFUNCTION(BlueprintCallable, Category = "Properties")
-	UE_API bool AddCustomLightProperty(UClass* InLightClass, const FName& InPropertyName);
+	UE_API bool AddCustomLightProperty(UClass* InLightComponentClass, const FName& InPropertyName);
 
 	/**
-	 * Removes the property in the specified class from the "Custom" lighting properties. Does nothing if the Custom properties doesn't contain the
+	 * Removes the property in the specified class from the "Custom" lighting properties. Does nothing if the Custom properties don't contain the
 	 * property. Returns true on success, else false.
 	 */
-	UE_API bool RemoveCustomLightProperty(const UClass* InLightClass, const FProperty* InLightingProperty);
+	UE_API bool RemoveCustomLightProperty(const UClass* InLightComponentClass, const FProperty* InLightProperty);
 
 	/**
-	 * Removes the property in the specified class from the "Custom" lighting properties. Does nothing if the Custom properties doesn't contain the
+	 * Removes the property in the specified class from the "Custom" lighting properties. Does nothing if the Custom properties don't contain the
 	 * property. Returns true on success, else false.
 	 */
 	UFUNCTION(BlueprintCallable, Category = "Properties")
-	UE_API bool RemoveCustomLightProperty(const UClass* InLightClass, const FName& InPropertyName);
+	UE_API bool RemoveCustomLightProperty(const UClass* InLightComponentClass, const FName& InPropertyName);
 
 	/** Determines if the given property exists within the "Custom" lighting properties. */
-	UE_API bool HasCustomLightProperty(const UClass* InLightClass, const FProperty* InLightingProperty) const;
+	UE_API bool HasCustomLightProperty(const UClass* InLightComponentClass, const FProperty* InLightProperty) const;
 	
 	/** Determines if the given property exists within the "Custom" lighting properties. */
 	UFUNCTION(BlueprintCallable, Category = "Properties")
-	UE_API bool HasCustomLightProperty(const UClass* InLightClass, const FName& InPropertyName) const;
+	UE_API bool HasCustomLightProperty(const UClass* InLightComponentClass, const FName& InPropertyName) const;
 
 	/**
 	 * Determines if the given "Custom" light property has been marked as overridden. A custom property MUST be marked as overridden in
 	 * order to take effect. The override state is visualized as the checkbox in front of the property name the UI.
 	 */
 	UFUNCTION(BlueprintCallable, Category = "Properties")
-	UE_API bool IsCustomLightPropertyOverridden(const UClass* InLightClass, const FName& InPropertyName) const;
+	UE_API bool IsCustomLightPropertyOverridden(const UClass* InLightComponentClass, const FName& InPropertyName) const;
 	
 	/**
 	 * Updates the override state for a "Custom" light property. A custom property MUST be marked as overridden in order to take effect. Returns true
 	 * on success, else false.
 	 */
 	UFUNCTION(BlueprintCallable, Category = "Properties")
-	UE_API bool UpdateCustomLightPropertyOverrideState(const UClass* InLightClass, const FName& InPropertyName, const bool bIsOverridden);
+	UE_API bool UpdateCustomLightPropertyOverrideState(const UClass* InLightComponentClass, const FName& InPropertyName, const bool bIsOverridden);
 
 	/** Gets how many "Custom" lighting properties have been added. */
 	UFUNCTION(BlueprintCallable, Category = "Properties")
@@ -312,7 +312,7 @@ public:
 
 private:
 	/** Gets the property name used for the given property within the DynamicProperties property bag. */
-	UE_API static FName GetLightingPropertyName(const UClass* InLightClass, const FProperty* InLightingProperty);
+	UE_API static FName GetLightingPropertyName(const UClass* InLightComponentClass, const FProperty* InLightProperty);
 
 	/**
 	 * Gets the bOverride_* property name used for the given property within the DynamicProperties property bag. The "main" property name
diff --git a/Engine/Plugins/MovieScene/MovieRenderPipeline/Source/MovieRenderPipelineEditor/Private/Customizations/Graph/MovieGraphLightModifierCustomization.cpp b/Engine/Plugins/MovieScene/MovieRenderPipeline/Source/MovieRenderPipelineEditor/Private/Customizations/Graph/MovieGraphLightModifierCustomization.cpp
--- a/Engine/Plugins/MovieScene/MovieRenderPipeline/Source/MovieRenderPipelineEditor/Private/Customizations/Graph/MovieGraphLightModifierCustomization.cpp
+++ b/Engine/Plugins/MovieScene/MovieRenderPipeline/Source/MovieRenderPipelineEditor/Private/Customizations/Graph/MovieGraphLightModifierCustomization.cpp
@@ -203,6 +203,16 @@ void FMovieGraphLightModifierCustomization::CustomizeDetails(const TSharedPtr<ID
 	CustomizeDetails(*InDetailBuilder);
 }
 
+FMovieGraphLightModifierCustomization::FCustomLightingPropertyInfo::FCustomLightingPropertyInfo(
+	UClass* LightActorClass, UClass* LightComponentClass, FProperty* LightProperty, const FString& LightPropertyCategory)
+	: LightActorClass(LightActorClass)
+	, LightComponentClass(LightComponentClass)
+	, LightProperty(LightProperty)
+	, LightPropertyCategory(LightPropertyCategory)
+{
+	
+}
+
 TWeakInterfacePtr<IMovieGraphModifierNodeInterface> FMovieGraphLightModifierCustomization::GetSelectedModifierNode() const
 {
 	if (const TSharedPtr<IDetailLayoutBuilder> DetailBuilderPin = DetailBuilder.Pin())
@@ -339,7 +349,7 @@ void FMovieGraphLightModifierCustomization::OnIntensityUnitsValueChange(const FN
 	}
 }
 
-TMap<UClass*, TMap<FString, TArray<const FProperty*>>> FMovieGraphLightModifierCustomization::GetCustomLightingProperties() const
+TArray<FMovieGraphLightModifierCustomization::FCustomLightingPropertyInfo> FMovieGraphLightModifierCustomization::GetCustomLightingProperties() const
 {
 	// Only properties in certain lighting categories should be included.
 	static const TArray<FString> AllowedLightingCategories = {
@@ -352,7 +362,7 @@ TMap<UClass*, TMap<FString, TArray<const FProperty*>>> FMovieGraphLightModifierC
 		TEXT("LightingChannels"), TEXT("SamplesPerPixel"), TEXT("Intensity"), TEXT("IntensityUnits")
 	};
 
-	TMap<UClass*, TMap<FString, TArray<const FProperty*>>> CustomLightingProperties;
+	TArray<FCustomLightingPropertyInfo> CustomLightingProperties;
 
 	// Get all light classes that are available. Add the Sky Light manually since it doesn't derive from ALight.
 	TArray<UClass*> LightClasses;
@@ -363,8 +373,6 @@ TMap<UClass*, TMap<FString, TArray<const FProperty*>>> FMovieGraphLightModifierC
 	// Discover all relevant lighting properties in the light classes.
 	for (UClass* LightClass : LightClasses)
 	{
-		TMap<FString, TArray<const FProperty*>>& LightClassProperties = CustomLightingProperties.FindOrAdd(LightClass);
-
 		// The light class itself doesn't contain the properties we care about, they're in a component (object property) in the class. Iterate the
 		// class until we find this component property.
 		const FObjectPropertyBase* LightComponentProperty = nullptr;
@@ -395,7 +403,7 @@ TMap<UClass*, TMap<FString, TArray<const FProperty*>>> FMovieGraphLightModifierC
 		// Iterate the properties of the light component.
 		for (TFieldIterator<FProperty> ComponentPropertyIt(LightComponentProperty->PropertyClass, EFieldIteratorFlags::IncludeSuper, EFieldIteratorFlags::ExcludeDeprecated); ComponentPropertyIt; ++ComponentPropertyIt)
 		{
-			const FProperty* ComponentProperty = *ComponentPropertyIt;
+			FProperty* ComponentProperty = *ComponentPropertyIt;
 			const FName PropertyName = ComponentProperty->GetFName();
 			
 			// Some properties are excluded because they're always shown outside of the Custom properties.
@@ -417,7 +425,8 @@ TMap<UClass*, TMap<FString, TArray<const FProperty*>>> FMovieGraphLightModifierC
 				continue;
 			}
 
-			LightClassProperties.FindOrAdd(Category).Add(ComponentProperty);
+			CustomLightingProperties.Add(
+				FCustomLightingPropertyInfo(LightClass, LightComponentProperty->PropertyClass, ComponentProperty, Category));
 		}
 	}
 
@@ -444,50 +453,62 @@ TSharedRef<SWidget> FMovieGraphLightModifierCustomization::GetCustomMenuContents
 	constexpr bool bNoIndent = true;
 	CustomMenu.AddWidget(CustomMenuHelpTextWidget, FText::GetEmpty(), bNoIndent);
 	
-	const TMap<UClass*, TMap<FString, TArray<const FProperty*>>> PropertiesToDisplay = GetCustomLightingProperties();
+	TArray<FCustomLightingPropertyInfo> PropertiesToDisplay = GetCustomLightingProperties();
 
-	TArray<UClass*> LightClasses;
-	PropertiesToDisplay.GetKeys(LightClasses);
+	// Group "Custom" properties by light class
+	TMap<UClass*, TArray<FCustomLightingPropertyInfo>> PropertiesByLightClass;
+	for (const FCustomLightingPropertyInfo& LightPropertyInfo : PropertiesToDisplay)
+	{
+		PropertiesByLightClass.FindOrAdd(LightPropertyInfo.LightActorClass).Add(LightPropertyInfo);
+	}
 
 	CustomMenu.BeginSection(NAME_None, LOCTEXT("LightActorTypesSectionHeader", "Light Component Types"));
 
-	for (UClass* LightClass : LightClasses)
+	for (const TPair<UClass*, TArray<FCustomLightingPropertyInfo>>& LightClassPropertiesPair : PropertiesByLightClass)
 	{
-		const TMap<FString, TArray<const FProperty*>>* LightClassProperties = PropertiesToDisplay.Find(LightClass);
-		if (!LightClassProperties)
-		{
-			continue;
-		}
+		const UClass* LightClass = LightClassPropertiesPair.Key;
+		const TArray<FCustomLightingPropertyInfo>& LightClassProperties = LightClassPropertiesPair.Value;
 
 		CustomMenu.AddSubMenu(
 			LightClass->GetDisplayNameText(),
 			TAttribute<FText>(),
-			FNewMenuDelegate::CreateLambda([LightClass, InWeakModifierNode, ClassProperties = *LightClassProperties](FMenuBuilder& MenuBuilder)
+			FNewMenuDelegate::CreateLambda([InWeakModifierNode, LightClassProperties](FMenuBuilder& MenuBuilder)
 			{
-				for (const TPair<FString, TArray<const FProperty*>>& CategoryPair : ClassProperties)
+				// Group properties by category
+				TMap<FString, TArray<FCustomLightingPropertyInfo>> LightPropertiesByCategory;
+				for (const FCustomLightingPropertyInfo& LightProperty : LightClassProperties)
 				{
-					MenuBuilder.BeginSection(NAME_None, FText::FromString(FName::NameToDisplayString(CategoryPair.Key, false)));
+					LightPropertiesByCategory.FindOrAdd(LightProperty.LightPropertyCategory).Add(LightProperty);
+				}
+
+				// Add properties for each property category
+				for (const TPair<FString, TArray<FCustomLightingPropertyInfo>>& LightPropertyPair : LightPropertiesByCategory)
+				{
+					const FString& LightCategory = LightPropertyPair.Key;
+					const TArray<FCustomLightingPropertyInfo>& CategoryProperties = LightPropertyPair.Value;
+
+					MenuBuilder.BeginSection(NAME_None, FText::FromString(FName::NameToDisplayString(LightCategory, false)));
 					
-					for (const FProperty* CategoryProperty : CategoryPair.Value)
+					for (const FCustomLightingPropertyInfo& PropertyInfo : CategoryProperties)
 					{
 						MenuBuilder.AddMenuEntry(
-							CategoryProperty->GetDisplayNameText(),
+							PropertyInfo.LightProperty->GetDisplayNameText(),
 							TAttribute<FText>(),
 							FSlateIcon(),
 							FUIAction(
-								FExecuteAction::CreateLambda([InWeakModifierNode, LightClass, CategoryProperty]()
+								FExecuteAction::CreateLambda([InWeakModifierNode, PropertyInfo]()
 								{
 									if (UMovieGraphLightModifierNode* LightModifierMode = Cast<UMovieGraphLightModifierNode>(InWeakModifierNode.Get()))
 									{
-										if (LightModifierMode->HasCustomLightProperty(LightClass, CategoryProperty))
+										if (LightModifierMode->HasCustomLightProperty(PropertyInfo.LightComponentClass, PropertyInfo.LightProperty))
 										{
 											const FScopedTransaction Transaction(LOCTEXT("RemoveCustomLightingProperty", "Remove Custom Lighting Property"));
-											LightModifierMode->RemoveCustomLightProperty(LightClass, CategoryProperty);
+											LightModifierMode->RemoveCustomLightProperty(PropertyInfo.LightComponentClass, PropertyInfo.LightProperty);
 										}
 										else
 										{
 											const FScopedTransaction Transaction(LOCTEXT("AddCustomLightingProperty", "Add Custom Lighting Property"));
-											LightModifierMode->AddCustomLightProperty(LightClass, CategoryProperty);
+											LightModifierMode->AddCustomLightProperty(PropertyInfo.LightComponentClass, PropertyInfo.LightProperty);
 										}
 									}
 								}),
@@ -495,11 +516,11 @@ TSharedRef<SWidget> FMovieGraphLightModifierCustomization::GetCustomMenuContents
 								{
 									return true;
 								}),
-								FIsActionChecked::CreateLambda([InWeakModifierNode, LightClass, CategoryProperty]()
+								FIsActionChecked::CreateLambda([InWeakModifierNode, PropertyInfo]()
 								{
 									if (const UMovieGraphLightModifierNode* LightModifierMode = Cast<UMovieGraphLightModifierNode>(InWeakModifierNode.Get()))
 									{
-										return LightModifierMode->HasCustomLightProperty(LightClass, CategoryProperty);
+										return LightModifierMode->HasCustomLightProperty(PropertyInfo.LightComponentClass, PropertyInfo.LightProperty);
 									}
 
 									return false;
diff --git a/Engine/Plugins/MovieScene/MovieRenderPipeline/Source/MovieRenderPipelineEditor/Private/Customizations/Graph/MovieGraphLightModifierCustomization.h b/Engine/Plugins/MovieScene/MovieRenderPipeline/Source/MovieRenderPipelineEditor/Private/Customizations/Graph/MovieGraphLightModifierCustomization.h
--- a/Engine/Plugins/MovieScene/MovieRenderPipeline/Source/MovieRenderPipelineEditor/Private/Customizations/Graph/MovieGraphLightModifierCustomization.h
+++ b/Engine/Plugins/MovieScene/MovieRenderPipeline/Source/MovieRenderPipelineEditor/Private/Customizations/Graph/MovieGraphLightModifierCustomization.h
@@ -23,6 +23,18 @@ protected:
 	//~ End IDetailCustomization interface
 
 private:
+	/** Stores information about a property returned from GetCustomLightingProperties(). */
+	struct FCustomLightingPropertyInfo
+	{
+		FCustomLightingPropertyInfo() = delete;
+		FCustomLightingPropertyInfo(UClass* LightActorClass, UClass* LightComponentClass, FProperty* LightProperty, const FString& LightPropertyCategory);
+
+		UClass* LightActorClass = nullptr;
+		UClass* LightComponentClass = nullptr;
+		FProperty* LightProperty = nullptr;
+		FString LightPropertyCategory;
+	};
+
 	/** Gets the selected modifier node that this customization is displaying. */
 	TWeakInterfacePtr<IMovieGraphModifierNodeInterface> GetSelectedModifierNode() const;
 
@@ -49,11 +61,8 @@ private:
 
 	/**
 	 * Gets all of the properties that should be displayed in the "Custom" properties menu.
-	 *
-	 * The key to the outermost TMap is the light class. The properties are contained in the inner map. In the inner map, the key is the property
-	 * category, the value is the property itself.
 	 */
-	TMap<UClass*, TMap<FString, TArray<const FProperty*>>> GetCustomLightingProperties() const;
+	TArray<FCustomLightingPropertyInfo> GetCustomLightingProperties() const;
 
 private:
 	/** The details builder associated with the customization. */
