MovieGraph: Fixed regression with the Select node caused by changes made for the Light Modifier node.

#Synced-CL 50977214

diff --git a/Engine/Plugins/MovieScene/MovieRenderPipeline/Source/MovieRenderPipelineCore/Public/Graph/MovieGraphValueView.cpp b/Engine/Plugins/MovieScene/MovieRenderPipeline/Source/MovieRenderPipelineCore/Public/Graph/MovieGraphValueView.cpp
--- a/Engine/Plugins/MovieScene/MovieRenderPipeline/Source/MovieRenderPipelineCore/Public/Graph/MovieGraphValueView.cpp
+++ b/Engine/Plugins/MovieScene/MovieRenderPipeline/Source/MovieRenderPipelineCore/Public/Graph/MovieGraphValueView.cpp
@@ -345,7 +345,18 @@ void UMovieGraphValueView::SetValueType(const FName& InPropertyName, EMovieGraph
 		Modify();
 
 		constexpr bool bOverwrite = true;
-		PropertyBag->AddProperty(InPropertyName, static_cast<EPropertyBagPropertyType>(ValueType), InValueTypeObject, bOverwrite);
+
+		// If the property is already a container, set the value type of the container. Otherwise, just set the value type of the property.
+		if (!Desc->ContainerTypes.IsEmpty())
+		{
+			PropertyBag->AddContainerProperty(
+				InPropertyName, Desc->ContainerTypes.GetFirstContainerType(), static_cast<EPropertyBagPropertyType>(ValueType),
+				InValueTypeObject, bOverwrite, Desc->KeyType, Desc->KeyTypeObject);
+		}
+		else
+		{
+			PropertyBag->AddProperty(InPropertyName, static_cast<EPropertyBagPropertyType>(ValueType), InValueTypeObject, bOverwrite);
+		}
 	}
 }
 
diff --git a/Engine/Plugins/MovieScene/MovieRenderPipeline/Source/MovieRenderPipelineEditor/Private/Customizations/Graph/MovieGraphSelectNodeCustomization.h b/Engine/Plugins/MovieScene/MovieRenderPipeline/Source/MovieRenderPipelineEditor/Private/Customizations/Graph/MovieGraphSelectNodeCustomization.h
--- a/Engine/Plugins/MovieScene/MovieRenderPipeline/Source/MovieRenderPipelineEditor/Private/Customizations/Graph/MovieGraphSelectNodeCustomization.h
+++ b/Engine/Plugins/MovieScene/MovieRenderPipeline/Source/MovieRenderPipelineEditor/Private/Customizations/Graph/MovieGraphSelectNodeCustomization.h
@@ -114,53 +114,63 @@ protected:
 			]
 			.ValueContent()
 			[
-				SNew(SPinTypeSelector, FGetPinTypeTree::CreateLambda(GetFilteredVariableTypeTree))
-				.TargetPinType_Lambda([SelectNodeObject]()
-				{
-					// The SPinTypeSelector popup might outlive this details view, so the node could be invalid
-					if (!SelectNodeObject.IsValid())
+				SNew(SBox)
+				.Padding(5, 0, 0, 0)
+				[
+					SNew(SPinTypeSelector, FGetPinTypeTree::CreateLambda(GetFilteredVariableTypeTree))
+					.TargetPinType_Lambda([SelectNodeObject]()
 					{
-						return FEdGraphPinType();
-					}
+						// The SPinTypeSelector popup might outlive this details view, so the node could be invalid
+						if (!SelectNodeObject.IsValid())
+						{
+							return FEdGraphPinType();
+						}
 
-					constexpr bool bIsBranch = false;
-					constexpr bool bIsWildcard = false;
-					return UMoviePipelineEdGraphNodeBase::GetPinType(SelectNodeObject->GetValueType(), bIsBranch, bIsWildcard, SelectNodeObject->GetValueTypeObject());
-				})
-				.OnPinTypeChanged_Lambda([PropUtils, SelectNodeObject](const FEdGraphPinType& InPinType)
-				{
-					// The SPinTypeSelector popup might outlive this details view, so the node could be invalid
-					if (SelectNodeObject.IsValid())
+						constexpr bool bIsBranch = false;
+						constexpr bool bIsWildcard = false;
+						return UMoviePipelineEdGraphNodeBase::GetPinType(SelectNodeObject->GetValueType(), bIsBranch, bIsWildcard, SelectNodeObject->GetValueTypeObject());
+					})
+					.OnPinTypeChanged_Lambda([PropUtils, SelectNodeObject](const FEdGraphPinType& InPinType)
 					{
-						SelectNodeObject->SetDataType(UMoviePipelineEdGraphNodeBase::GetValueTypeFromPinType(InPinType), InPinType.PinSubCategoryObject.Get());
+						// The SPinTypeSelector popup might outlive this details view, so the node could be invalid
+						if (SelectNodeObject.IsValid())
+						{
+							SelectNodeObject->SetDataType(UMoviePipelineEdGraphNodeBase::GetValueTypeFromPinType(InPinType), InPinType.PinSubCategoryObject.Get());
 
-						// Need the ForceRefresh to make sure the details panel refreshes immediately after the data type change.
-						// Can result in a crash without it.
-						PropUtils->ForceRefresh();
-					}
-				})
-				.Schema(GetDefault<UPropertyBagSchema>())
-				.bAllowArrays(false)
-				.TypeTreeFilter(ETypeTreeFilter::None)
-				.Font(IDetailLayoutBuilder::GetDetailFont())
+							// Need the ForceRefresh to make sure the details panel refreshes immediately after the data type change.
+							// Can result in a crash without it.
+							PropUtils->ForceRefresh();
+						}
+					})
+					.Schema(GetDefault<UPropertyBagSchema>())
+					.bAllowArrays(false)
+					.TypeTreeFilter(ETypeTreeFilter::None)
+					.Font(IDetailLayoutBuilder::GetDetailFont())
+				]
 			];
 
 			const TSharedRef<IPropertyHandle> SelectOptions = DetailBuilder.GetProperty("SelectOptions");
 			const TSharedRef<IPropertyHandle> SelectedOption = DetailBuilder.GetProperty("SelectedOption");
 
-			// Add the property bag property ("Value") for both SelectOptions and SelectedOption
+			// Add the property bag property ("MultiValueContainer.Values") for both SelectOptions and SelectedOption
 			if (SelectOptions->IsValidHandle() && SelectedOption->IsValidHandle())
 			{
-				if (const TSharedPtr<IPropertyHandle> ValueProperty = SelectOptions->GetChildHandle("Value"))
+				if (const TSharedPtr<IPropertyHandle> ContainerProperty = SelectOptions->GetChildHandle("MultiValueContainer"))
 				{
-					Category.AddProperty(ValueProperty);
-					SelectOptions->MarkHiddenByCustomization();
+					if (const TSharedPtr<IPropertyHandle> ValuesProperty = ContainerProperty->GetChildHandle("Values"))
+					{
+						Category.AddProperty(ValuesProperty);
+						SelectOptions->MarkHiddenByCustomization();
+					}
 				}
 
-				if (const TSharedPtr<IPropertyHandle> ValueProperty = SelectedOption->GetChildHandle("Value"))
+				if (const TSharedPtr<IPropertyHandle> ContainerProperty = SelectedOption->GetChildHandle("MultiValueContainer"))
 				{
-					Category.AddProperty(ValueProperty);
-					SelectedOption->MarkHiddenByCustomization();
+					if (const TSharedPtr<IPropertyHandle> ValuesProperty = ContainerProperty->GetChildHandle("Values"))
+					{
+						Category.AddProperty(ValuesProperty);
+						SelectedOption->MarkHiddenByCustomization();
+					}
 				}
 			}
 		}
