Movie Pipeline: Bypass muted cameras when all cameras is enabled

#Synced-CL 50973831

diff --git a/Engine/Plugins/MovieScene/MovieRenderPipeline/Source/MovieRenderPipelineCore/Private/MoviePipelineBlueprintLibrary.cpp b/Engine/Plugins/MovieScene/MovieRenderPipeline/Source/MovieRenderPipelineCore/Private/MoviePipelineBlueprintLibrary.cpp
--- a/Engine/Plugins/MovieScene/MovieRenderPipeline/Source/MovieRenderPipelineCore/Private/MoviePipelineBlueprintLibrary.cpp
+++ b/Engine/Plugins/MovieScene/MovieRenderPipeline/Source/MovieRenderPipelineCore/Private/MoviePipelineBlueprintLibrary.cpp
@@ -38,6 +38,7 @@
 #include "Features/IModularFeatures.h"
 #include "Misc/EngineVersion.h"
 #include "GenericPlatform/GenericPlatformProcess.h"
+#include "Decorations/MovieSceneMuteSoloDecoration.h"
 
 // For camera settings
 #include "GameFramework/PlayerController.h"
@@ -562,7 +563,7 @@ void UMoviePipelineBlueprintLibrary::UpdateJobShotListFromSequence(ULevelSequenc
 						SidecarCamera.BindingId = ValidBinding;
 						SidecarCamera.SequenceId = Entity.SequenceID; 
 						bool bAddCamera = true;
-						
+
 						TArray<FGuid> CheckedBindings;
 						
 						// We break out of the while loop below when the parent guid is invalid, so we track
@@ -573,6 +574,19 @@ void UMoviePipelineBlueprintLibrary::UpdateJobShotListFromSequence(ULevelSequenc
 
 						while (ParentGuid.IsValid())
 						{
+							// Bypass if the binding is muted
+							if (FMovieSceneBinding* Binding = LeafNode->MovieScene->FindBinding(ParentGuid))
+							{
+								if (UMovieSceneMuteSoloDecoration* MuteSoloDecoration = Cast<UMovieSceneMuteSoloDecoration>(Binding->FindDecoration<UMovieSceneMuteSoloDecoration>()))
+								{
+									if (MuteSoloDecoration->IsMuted())
+									{
+										bAddCamera = false;
+										break;
+									}
+								}
+							}
+
 							if (FMovieScenePossessable* ParentAsPossessable = LeafNode->MovieScene->FindPossessable(ParentGuid))
 							{
 								LastValidGuid = ParentGuid;
@@ -583,6 +597,7 @@ void UMoviePipelineBlueprintLibrary::UpdateJobShotListFromSequence(ULevelSequenc
 								if (MovieSceneHelpers::IsBoundToAnySpawnable(LeafNode->MovieScene->GetTypedOuter<UMovieSceneSequence>(), ParentAsPossessable->GetGuid(), TransientPlaybackState))
 								{
 									FMovieSceneBinding* Binding = LeafNode->MovieScene->FindBinding(ParentAsPossessable->GetGuid());
+
 									if (Binding)
 									{
 										const TArray<UMovieSceneTrack*> AllBindingTracks = Binding->GetTracks();
