MoviePipeline: Fixed an issue that caused "Named Resolution"-typed overrides in the Quick Render window to not take effect.

#Synced-CL 51044521

diff --git a/Engine/Plugins/MovieScene/MovieRenderPipeline/Source/MovieRenderPipelineEditor/Private/Graph/MovieGraphAssetToolkit.cpp b/Engine/Plugins/MovieScene/MovieRenderPipeline/Source/MovieRenderPipelineEditor/Private/Graph/MovieGraphAssetToolkit.cpp
--- a/Engine/Plugins/MovieScene/MovieRenderPipeline/Source/MovieRenderPipelineEditor/Private/Graph/MovieGraphAssetToolkit.cpp
+++ b/Engine/Plugins/MovieScene/MovieRenderPipeline/Source/MovieRenderPipelineEditor/Private/Graph/MovieGraphAssetToolkit.cpp
@@ -11,7 +11,6 @@
 #include "Customizations/Graph/MovieGraphMemberCustomization.h"
 #include "Customizations/Graph/MovieGraphMetadataAttributeCustomization.h"
 #include "Customizations/Graph/MovieGraphModifiersCustomization.h"
-#include "Customizations/Graph/MovieGraphNamedResolutionCustomization.h"
 #include "Customizations/Graph/MovieGraphNodeCustomization.h"
 #include "Customizations/Graph/MovieGraphPathTracedRendererNodeCustomization.h"
 #include "Customizations/Graph/MovieGraphPostProcessPassCustomization.h"
@@ -414,10 +413,6 @@ TSharedRef<SDockTab> FMovieGraphAssetToolkit::SpawnTab_RenderGraphDetails(const
 		UMovieGraphShowFlags::StaticClass()->GetFName(),
 		FOnGetPropertyTypeCustomizationInstance::CreateStatic(&FMovieGraphShowFlagsCustomization::MakeInstance));
 
-	SelectedGraphObjectsDetailsWidget->RegisterInstancedCustomPropertyTypeLayout(
-		FMovieGraphNamedResolution::StaticStruct()->GetFName(),
-		FOnGetPropertyTypeCustomizationInstance::CreateStatic(&FMovieGraphNamedResolutionCustomization::MakeInstance));
-
 	SelectedGraphObjectsDetailsWidget->RegisterInstancedCustomPropertyTypeLayout(
 		FMovieGraphMetadataAttribute::StaticStruct()->GetFName(),
 		FOnGetPropertyTypeCustomizationInstance::CreateStatic(&FMovieGraphMetadataAttributeCustomization::MakeInstance));
diff --git a/Engine/Plugins/MovieScene/MovieRenderPipeline/Source/MovieRenderPipelineEditor/Private/MovieRenderPipelineEditorModule.cpp b/Engine/Plugins/MovieScene/MovieRenderPipeline/Source/MovieRenderPipelineEditor/Private/MovieRenderPipelineEditorModule.cpp
--- a/Engine/Plugins/MovieScene/MovieRenderPipeline/Source/MovieRenderPipelineEditor/Private/MovieRenderPipelineEditorModule.cpp
+++ b/Engine/Plugins/MovieScene/MovieRenderPipeline/Source/MovieRenderPipelineEditor/Private/MovieRenderPipelineEditorModule.cpp
@@ -36,9 +36,11 @@
 #include "MoviePipelineConsoleVariableSetting.h"
 #include "PropertyEditorModule.h"
 #include "Customizations/ConsoleVariableCustomization.h"
+#include "Customizations/Graph/MovieGraphNamedResolutionCustomization.h"
 #include "LevelSequence.h"
 #include "ToolMenus.h"
 #include "Graph/MovieGraphDataTypes.h"
+#include "Graph/MovieGraphNamedResolution.h"
 #include "Graph/MovieGraphPinFactory.h"
 
 #define LOCTEXT_NAMESPACE "FMovieRenderPipelineEditorModule"
@@ -219,8 +221,12 @@ void FMovieRenderPipelineEditorModule::UnregisterMovieRenderer()
 void FMovieRenderPipelineEditorModule::RegisterTypeCustomizations()
 {
 	FPropertyEditorModule& PropertyModule = FModuleManager::LoadModuleChecked<FPropertyEditorModule>("PropertyEditor");
+	
 	PropertyModule.RegisterCustomPropertyTypeLayout(FMoviePipelineConsoleVariableEntry::StaticStruct()->GetFName(),
 		FOnGetPropertyTypeCustomizationInstance::CreateStatic(FConsoleVariablesDetailsCustomization::MakeInstance));
+
+	PropertyModule.RegisterCustomPropertyTypeLayout(FMovieGraphNamedResolution::StaticStruct()->GetFName(),
+		FOnGetPropertyTypeCustomizationInstance::CreateStatic(&FMovieGraphNamedResolutionCustomization::MakeInstance));
 }
 
 void FMovieRenderPipelineEditorModule::UnregisterTypeCustomizations()
@@ -231,6 +237,7 @@ void FMovieRenderPipelineEditorModule::UnregisterTypeCustomizations()
 		if (PropertyModule)
 		{
 			PropertyModule->UnregisterCustomPropertyTypeLayout(FMoviePipelineConsoleVariableEntry::StaticStruct()->GetFName());
+			PropertyModule->UnregisterCustomPropertyTypeLayout(FMovieGraphNamedResolution::StaticStruct()->GetFName());
 		}
 	}
 }
