MoviePipeline: Speculative fix for a crash caused by MRG global variable constructors, likely due to having their logic not execute on the game thread in some scenarios.

#Synced-CL 51044521

diff --git a/Engine/Plugins/MovieScene/MovieRenderPipeline/Source/MovieRenderPipelineCore/Private/Graph/MovieGraphConfig.cpp b/Engine/Plugins/MovieScene/MovieRenderPipeline/Source/MovieRenderPipelineCore/Private/Graph/MovieGraphConfig.cpp
--- a/Engine/Plugins/MovieScene/MovieRenderPipeline/Source/MovieRenderPipelineCore/Private/Graph/MovieGraphConfig.cpp
+++ b/Engine/Plugins/MovieScene/MovieRenderPipeline/Source/MovieRenderPipelineCore/Private/Graph/MovieGraphConfig.cpp
@@ -274,25 +274,57 @@ bool UMovieGraphGlobalVariable::CanRename(const FText& InNewName, FText& OutErro
 UMovieGraphGlobalVariable_ShotName::UMovieGraphGlobalVariable_ShotName()
 {
 	Name = FString(TEXT("shot_name"));
-	SetValueType(EMovieGraphValueType::String);
+
+	// SetValueType() will result in a PostEditChangeProperty call; ensure it's only called on the game thread
+	AsyncTask(ENamedThreads::GameThread, [WeakThis = TWeakObjectPtr(this)]()
+	{
+		if (UMovieGraphGlobalVariable_ShotName* ShotVariable = WeakThis.Get())
+		{
+			ShotVariable->SetValueType(EMovieGraphValueType::String);
+		}
+	});
 }
 
 UMovieGraphGlobalVariable_SequenceName::UMovieGraphGlobalVariable_SequenceName()
 {
 	Name = FString(TEXT("seq_name"));
-	SetValueType(EMovieGraphValueType::String);
+
+	// SetValueType() will result in a PostEditChangeProperty call; ensure it's only called on the game thread
+	AsyncTask(ENamedThreads::GameThread, [WeakThis = TWeakObjectPtr(this)]()
+	{
+		if (UMovieGraphGlobalVariable_SequenceName* SequenceVariable = WeakThis.Get())
+		{
+			SequenceVariable->SetValueType(EMovieGraphValueType::String);
+		}
+	});
 }
 
 UMovieGraphGlobalVariable_FrameNumber::UMovieGraphGlobalVariable_FrameNumber()
 {
 	Name = FString(TEXT("frame_num"));
-	SetValueType(EMovieGraphValueType::Int32);
+
+	// SetValueType() will result in a PostEditChangeProperty call; ensure it's only called on the game thread
+	AsyncTask(ENamedThreads::GameThread, [WeakThis = TWeakObjectPtr(this)]()
+	{
+		if (UMovieGraphGlobalVariable_FrameNumber* FrameVariable = WeakThis.Get())
+		{
+			FrameVariable->SetValueType(EMovieGraphValueType::Int32);
+		}
+	});
 }
 
 UMovieGraphGlobalVariable_CameraName::UMovieGraphGlobalVariable_CameraName()
 {
 	Name = FString(TEXT("camera_name"));
-	SetValueType(EMovieGraphValueType::String);
+
+	// SetValueType() will result in a PostEditChangeProperty call; ensure it's only called on the game thread
+	AsyncTask(ENamedThreads::GameThread, [WeakThis = TWeakObjectPtr(this)]()
+	{
+		if (UMovieGraphGlobalVariable_CameraName* CameraVariable = WeakThis.Get())
+		{
+			CameraVariable->SetValueType(EMovieGraphValueType::String);
+		}
+	});
 }
 
 void UMovieGraphGlobalVariable_ShotName::UpdateValue(const FMovieGraphTraversalContext* InTraversalContext, const UMovieGraphPipeline* InPipeline)
