Movie Render Graph: Convert to graph if project settings call for it

#Synced-CL 51308618

diff --git a/Engine/Plugins/MovieScene/MovieRenderPipeline/Source/MovieRenderPipelineEditor/Private/MovieRenderPipelineEditorModule.cpp b/Engine/Plugins/MovieScene/MovieRenderPipeline/Source/MovieRenderPipelineEditor/Private/MovieRenderPipelineEditorModule.cpp
--- a/Engine/Plugins/MovieScene/MovieRenderPipeline/Source/MovieRenderPipelineEditor/Private/MovieRenderPipelineEditorModule.cpp
+++ b/Engine/Plugins/MovieScene/MovieRenderPipeline/Source/MovieRenderPipelineEditor/Private/MovieRenderPipelineEditorModule.cpp
@@ -42,6 +42,7 @@
 #include "Graph/MovieGraphDataTypes.h"
 #include "Graph/MovieGraphNamedResolution.h"
 #include "Graph/MovieGraphPinFactory.h"
+#include "Graph/MovieGraphPipeline.h"
 
 #define LOCTEXT_NAMESPACE "FMovieRenderPipelineEditorModule"
 
@@ -149,6 +150,17 @@ class FMovieRenderPipelineRenderer : public IMovieRendererInterface
 			}
 
 			UMoviePipelineEditorBlueprintLibrary::EnsureJobHasDefaultSettings(ActiveJob);
+
+			// If the default class is a movie graph, assign the default graph
+			const UClass* DefaultPipeline = Cast<UClass>(ProjectSettings->DefaultPipeline.TryLoad());
+			if (DefaultPipeline && DefaultPipeline == UMovieGraphPipeline::StaticClass())
+			{
+				const TSoftObjectPtr<UMovieGraphConfig> ProjectDefaultGraph = ProjectSettings->DefaultGraph;
+				if (const UMovieGraphConfig* DefaultGraph = ProjectDefaultGraph.LoadSynchronous())
+				{
+					ActiveJob->SetGraphPreset(DefaultGraph);
+				}
+			}
 		}
 
 		TArray<FString> SequenceNames;
