MoviePipeline: No longer set r.SkipRedundantTransformUpdate in MRQ/MRG as it has been deprecated.

#Synced-CL 51044521

diff --git a/Engine/Plugins/MovieScene/MovieRenderPipeline/Source/MovieRenderPipelineCore/Private/Graph/Nodes/MovieGraphGlobalGameOverrides.cpp b/Engine/Plugins/MovieScene/MovieRenderPipeline/Source/MovieRenderPipelineCore/Private/Graph/Nodes/MovieGraphGlobalGameOverrides.cpp
--- a/Engine/Plugins/MovieScene/MovieRenderPipeline/Source/MovieRenderPipelineCore/Private/Graph/Nodes/MovieGraphGlobalGameOverrides.cpp
+++ b/Engine/Plugins/MovieScene/MovieRenderPipeline/Source/MovieRenderPipelineCore/Private/Graph/Nodes/MovieGraphGlobalGameOverrides.cpp
@@ -97,7 +97,6 @@ void UMovieGraphGlobalGameOverridesNode::BuildNewProcessCommandLineArgsImpl(TArr
 	InOutDeviceProfileCvars.AddUnique(FString::Printf(TEXT("r.VolumetricRenderTarget.Mode=%d"), 3));
 	InOutDeviceProfileCvars.AddUnique(FString::Printf(TEXT("wp.Runtime.BlockOnSlowStreaming=%d"), 0));
 	InOutDeviceProfileCvars.AddUnique(FString::Printf(TEXT("p.Chaos.ImmPhys.MinStepTime=%d"), 0));
-	InOutDeviceProfileCvars.AddUnique(FString::Printf(TEXT("r.SkipRedundantTransformUpdate=%d"), 0));
 	InOutDeviceProfileCvars.AddUnique(FString::Printf(TEXT("p.ChaosCloth.UseTimeStepSmoothing=%d"), 0));
 	InOutDeviceProfileCvars.AddUnique(FString::Printf(TEXT("r.Water.SkipWaterInfoTextureRenderWhenWorldRenderingDisabled=%d"), 0));
 	InOutDeviceProfileCvars.AddUnique(FString::Printf(TEXT("Slate.bAllowThrottling=%d"), 0));
@@ -219,11 +218,6 @@ void UMovieGraphGlobalGameOverridesNode::ApplySettings(const bool bOverrideValue
 		// Remove any minimum delta time requirements from Chaos Physics to ensure accuracy at high Temporal Sample counts
 		MOVIEPIPELINE_STORE_AND_OVERRIDE_CVAR_FLOAT(NodeCDO->PreviousChaosImmPhysicsMinStepTime, TEXT("p.Chaos.ImmPhys.MinStepTime"), 0, bOverrideValues);
 
-		// MRQ's 0 -> 0.99 -> 0 evaluation for motion blur emulation can occasionally cause it to be detected as a redundant update and thus never updated
-		// which causes objects to render in the wrong position on the first frame (and without motion blur). This disables an optimization that detects
-		// the redundant updates so the update will get sent through anyways even though it thinks it's a duplicate (but it's not).
-		MOVIEPIPELINE_STORE_AND_OVERRIDE_CVAR_INT(NodeCDO->PreviousSkipRedundantTransformUpdate, TEXT("r.SkipRedundantTransformUpdate"), 0, bOverrideValues);
-
 		// Cloth's time step smoothing messes up the change in number of simulation substeps that fixes the cloth simulation behavior when using Temporal Samples.
 		MOVIEPIPELINE_STORE_AND_OVERRIDE_CVAR_INT(NodeCDO->PreviousChaosClothUseTimeStepSmoothing, TEXT("p.ChaosCloth.UseTimeStepSmoothing"), 0, bOverrideValues);
 
diff --git a/Engine/Plugins/MovieScene/MovieRenderPipeline/Source/MovieRenderPipelineCore/Private/MoviePipelineGameOverrideSetting.cpp b/Engine/Plugins/MovieScene/MovieRenderPipeline/Source/MovieRenderPipelineCore/Private/MoviePipelineGameOverrideSetting.cpp
--- a/Engine/Plugins/MovieScene/MovieRenderPipeline/Source/MovieRenderPipelineCore/Private/MoviePipelineGameOverrideSetting.cpp
+++ b/Engine/Plugins/MovieScene/MovieRenderPipeline/Source/MovieRenderPipelineCore/Private/MoviePipelineGameOverrideSetting.cpp
@@ -161,11 +161,6 @@ void UMoviePipelineGameOverrideSetting::ApplyCVarSettings(const bool bOverrideVa
 	// Remove any minimum delta time requirements from Chaos Physics to ensure accuracy at high Temporal Sample counts
 	MOVIEPIPELINE_STORE_AND_OVERRIDE_CVAR_FLOAT(PreviousChaosImmPhysicsMinStepTime, TEXT("p.Chaos.ImmPhys.MinStepTime"), 0, bOverrideValues);
 
-	// MRQ's 0 -> 0.99 -> 0 evaluation for motion blur emulation can occasionally cause it to be detected as a redundant update and thus never updated
-	// which causes objects to render in the wrong position on the first frame (and without motion blur). This disables an optimization that detects
-	// the redundant updates so the update will get sent through anyways even though it thinks it's a duplicate (but it's not).
-	MOVIEPIPELINE_STORE_AND_OVERRIDE_CVAR_INT(PreviousSkipRedundantTransformUpdate, TEXT("r.SkipRedundantTransformUpdate"), 0, bOverrideValues);
-
 	// Cloth's time step smoothing messes up the change in number of simulation substeps that fixes the cloth simulation behavior when using Temporal Samples.
 	MOVIEPIPELINE_STORE_AND_OVERRIDE_CVAR_INT(PreviousChaosClothUseTimeStepSmoothing, TEXT("p.ChaosCloth.UseTimeStepSmoothing"), 0, bOverrideValues);
 
@@ -302,7 +297,6 @@ void UMoviePipelineGameOverrideSetting::BuildNewProcessCommandLineArgsImpl(TArra
 	InOutDeviceProfileCvars.AddUnique(FString::Printf(TEXT("r.VolumetricRenderTarget.Mode=%d"), 3));
 	InOutDeviceProfileCvars.AddUnique(FString::Printf(TEXT("wp.Runtime.BlockOnSlowStreaming=%d"), 0));
 	InOutDeviceProfileCvars.AddUnique(FString::Printf(TEXT("p.Chaos.ImmPhys.MinStepTime=%d"), 0));
-	InOutDeviceProfileCvars.AddUnique(FString::Printf(TEXT("r.SkipRedundantTransformUpdate=%d"), 0));
 	InOutDeviceProfileCvars.AddUnique(FString::Printf(TEXT("p.ChaosCloth.UseTimeStepSmoothing=%d"), 0));
 	InOutDeviceProfileCvars.AddUnique(FString::Printf(TEXT("r.Water.SkipWaterInfoTextureRenderWhenWorldRenderingDisabled=%d"), 0));
 	InOutDeviceProfileCvars.AddUnique(FString::Printf(TEXT("Slate.bAllowThrottling=%d"), 0));
