Movie Pipeline: Follow up to 51035890 to also check whether the track is disabled (vs muted)

#Synced-CL 51424575

diff --git a/Engine/Plugins/MovieScene/MovieRenderPipeline/Source/MovieRenderPipelineCore/Private/MoviePipelineBlueprintLibrary.cpp b/Engine/Plugins/MovieScene/MovieRenderPipeline/Source/MovieRenderPipelineCore/Private/MoviePipelineBlueprintLibrary.cpp
--- a/Engine/Plugins/MovieScene/MovieRenderPipeline/Source/MovieRenderPipelineCore/Private/MoviePipelineBlueprintLibrary.cpp
+++ b/Engine/Plugins/MovieScene/MovieRenderPipeline/Source/MovieRenderPipelineCore/Private/MoviePipelineBlueprintLibrary.cpp
@@ -582,6 +582,7 @@ void UMoviePipelineBlueprintLibrary::UpdateJobShotListFromSequence(ULevelSequenc
 									if (MuteSoloDecoration->IsMuted())
 									{
 										bAddCamera = false;
+										UE_LOG(LogMovieRenderPipeline, Log, TEXT("Skipping muted camera: %s"), *Possessable.GetName());
 										break;
 									}
 								}
@@ -603,6 +604,14 @@ void UMoviePipelineBlueprintLibrary::UpdateJobShotListFromSequence(ULevelSequenc
 										const TArray<UMovieSceneTrack*> AllBindingTracks = Binding->GetTracks();
 										for (const UMovieSceneTrack* Track : AllBindingTracks)
 										{
+											// Bypass if the spawn track is disabled
+											if (Track->IsEvalDisabled())
+											{
+												bAddCamera = false;
+												UE_LOG(LogMovieRenderPipeline, Log, TEXT("Skipping disabled camera: %s"), *ParentAsPossessable->GetName());
+												continue;
+											}
+
 											if (const UMovieSceneSpawnTrack* SpawnTrack = Cast<UMovieSceneSpawnTrack>(Track))
 											{
 												for (const UMovieSceneSection* Section : SpawnTrack->GetAllSections())
