MoviePipeline: Fixed a bug where GetValueObject() and GetValueStruct() were not setting their OutValue properly in some cases (across UMovieGraphValueContainer, UMovieGraphMultiValueContainer, UMovieGraphValueView, and UMovieJobVariableAssignmentContainer).

With the fix, there is a slight chance that some C++ code may break. Additionally, some Blueprints may break due to the change to GetValueObject(). Although in both cases, the code/BP wouldn't have been working anyway, so this change may bring attention to that.

#Synced-CL 51044521

diff --git a/Engine/Plugins/MovieScene/MovieRenderPipeline/Source/MovieRenderPipelineCore/Private/MovieJobVariableAssignmentContainer.cpp b/Engine/Plugins/MovieScene/MovieRenderPipeline/Source/MovieRenderPipelineCore/Private/MovieJobVariableAssignmentContainer.cpp
--- a/Engine/Plugins/MovieScene/MovieRenderPipeline/Source/MovieRenderPipelineCore/Private/MovieJobVariableAssignmentContainer.cpp
+++ b/Engine/Plugins/MovieScene/MovieRenderPipeline/Source/MovieRenderPipelineCore/Private/MovieJobVariableAssignmentContainer.cpp
@@ -105,7 +105,7 @@ bool UMovieJobVariableAssignmentContainer::GetValueStruct(const UMovieGraphVaria
 	return UE::MovieGraph::Private::GetOptionalValue<FStructView>(Result, OutValue);
 }
 
-bool UMovieJobVariableAssignmentContainer::GetValueObject(const UMovieGraphVariable* InGraphVariable, UObject* OutValue, const UClass* RequestedClass) const
+bool UMovieJobVariableAssignmentContainer::GetValueObject(const UMovieGraphVariable* InGraphVariable, UObject*& OutValue, const UClass* RequestedClass) const
 {
 	TValueOrError<UObject*, EPropertyBagResult> Result = Value.GetValueObject(ConvertVariableToInternalName(InGraphVariable), RequestedClass);
 	return UE::MovieGraph::Private::GetOptionalValue<UObject*>(Result, OutValue);
diff --git a/Engine/Plugins/MovieScene/MovieRenderPipeline/Source/MovieRenderPipelineCore/Public/Graph/MovieGraphMultiValueContainer.h b/Engine/Plugins/MovieScene/MovieRenderPipeline/Source/MovieRenderPipelineCore/Public/Graph/MovieGraphMultiValueContainer.h
--- a/Engine/Plugins/MovieScene/MovieRenderPipeline/Source/MovieRenderPipelineCore/Public/Graph/MovieGraphMultiValueContainer.h
+++ b/Engine/Plugins/MovieScene/MovieRenderPipeline/Source/MovieRenderPipelineCore/Public/Graph/MovieGraphMultiValueContainer.h
@@ -119,7 +119,7 @@ public:
 
 	/** Gets the struct value of the held data. Returns true on success, else false. */
 	template <typename T>
-	bool GetValueStruct(const FName& InPropertyName, T* OutValue) const
+	bool GetValueStruct(const FName& InPropertyName, T*& OutValue) const
 	{
 		return ValuesView->GetValueStruct<T>(InPropertyName, OutValue);
 	}
diff --git a/Engine/Plugins/MovieScene/MovieRenderPipeline/Source/MovieRenderPipelineCore/Public/Graph/MovieGraphValueContainer.h b/Engine/Plugins/MovieScene/MovieRenderPipeline/Source/MovieRenderPipelineCore/Public/Graph/MovieGraphValueContainer.h
--- a/Engine/Plugins/MovieScene/MovieRenderPipeline/Source/MovieRenderPipelineCore/Public/Graph/MovieGraphValueContainer.h
+++ b/Engine/Plugins/MovieScene/MovieRenderPipeline/Source/MovieRenderPipelineCore/Public/Graph/MovieGraphValueContainer.h
@@ -90,7 +90,7 @@ public:
 
 	/** Gets the struct value of the held data. Returns true on success, else false. */
 	template <typename T>
-	bool GetValueStruct(T* OutValue) const
+	bool GetValueStruct(T*& OutValue) const
 	{
 		return MultiValueContainer->GetValueStruct(PropertyName, OutValue);
 	}
diff --git a/Engine/Plugins/MovieScene/MovieRenderPipeline/Source/MovieRenderPipelineCore/Public/Graph/MovieGraphValueView.h b/Engine/Plugins/MovieScene/MovieRenderPipeline/Source/MovieRenderPipelineCore/Public/Graph/MovieGraphValueView.h
--- a/Engine/Plugins/MovieScene/MovieRenderPipeline/Source/MovieRenderPipelineCore/Public/Graph/MovieGraphValueView.h
+++ b/Engine/Plugins/MovieScene/MovieRenderPipeline/Source/MovieRenderPipelineCore/Public/Graph/MovieGraphValueView.h
@@ -187,7 +187,7 @@ public:
 
 	/** Gets the struct value of the held data. Returns true on success, else false. */
 	template <typename T>
-	bool GetValueStruct(const FName& InPropertyName, T* OutValue) const
+	bool GetValueStruct(const FName& InPropertyName, T*& OutValue) const
 	{
 		TValueOrError<T*, EPropertyBagResult> Result = PropertyBag->GetValueStruct<T*>(InPropertyName);
 		return UE::MovieGraph::Private::GetOptionalValue<T*>(Result, OutValue);
diff --git a/Engine/Plugins/MovieScene/MovieRenderPipeline/Source/MovieRenderPipelineCore/Public/MovieJobVariableAssignmentContainer.h b/Engine/Plugins/MovieScene/MovieRenderPipeline/Source/MovieRenderPipelineCore/Public/MovieJobVariableAssignmentContainer.h
--- a/Engine/Plugins/MovieScene/MovieRenderPipeline/Source/MovieRenderPipelineCore/Public/MovieJobVariableAssignmentContainer.h
+++ b/Engine/Plugins/MovieScene/MovieRenderPipeline/Source/MovieRenderPipelineCore/Public/MovieJobVariableAssignmentContainer.h
@@ -85,7 +85,7 @@ public:
 
 	/** Gets the object value (for a specific class) of the specified property. Returns true on success, else false. */
 	UFUNCTION(BlueprintCallable, Category="Value")
-	UE_API bool GetValueObject(const UMovieGraphVariable* InGraphVariable, UObject* OutValue, const UClass* RequestedClass = nullptr) const;
+	UE_API bool GetValueObject(const UMovieGraphVariable* InGraphVariable, UObject*& OutValue, const UClass* RequestedClass = nullptr) const;
 
 	/** Gets the UClass value of the specified property. Returns true on success, else false. */
 	UFUNCTION(BlueprintCallable, Category="Value")
@@ -105,7 +105,7 @@ public:
 
 	/** Gets the struct value of the specified property. Returns true on success, else false. */
 	template <typename T>
-	bool GetValueStruct(const UMovieGraphVariable* InGraphVariable, T* OutValue) const
+	bool GetValueStruct(const UMovieGraphVariable* InGraphVariable, T*& OutValue) const
 	{
 		TValueOrError<T*, EPropertyBagResult> Result = Value.GetValueStruct<T*>(ConvertVariableToInternalName(InGraphVariable));
 		return UE::MovieGraph::Private::GetOptionalValue<T*>(Result, OutValue);
@@ -113,7 +113,7 @@ public:
 
 	/** Gets the object value of the specified property. Returns true on success, else false. */
 	template <typename T>
-	bool GetValueObject(const UMovieGraphVariable* InGraphVariable, T* OutValue) const
+	bool GetValueObject(const UMovieGraphVariable* InGraphVariable, T*& OutValue) const
 	{
 		TValueOrError<T*, EPropertyBagResult> Result = Value.GetValueObject<T*>(ConvertVariableToInternalName(InGraphVariable));
 		return UE::MovieGraph::Private::GetOptionalValue<T*>(Result, OutValue);
