MoviePipeline: Fixed a bug which could cause non-deterministic collection evaluation in the graph. This could happen in a scenario where a modifier is modifying actors across several sublevels, and actor names are duplicated across those sublevels.

#Synced-CL 51651671

diff --git a/Engine/Plugins/MovieScene/MovieRenderPipeline/Source/MovieRenderPipelineCore/Private/Graph/MovieGraphRenderLayerSubsystem.cpp b/Engine/Plugins/MovieScene/MovieRenderPipeline/Source/MovieRenderPipelineCore/Private/Graph/MovieGraphRenderLayerSubsystem.cpp
--- a/Engine/Plugins/MovieScene/MovieRenderPipeline/Source/MovieRenderPipelineCore/Private/Graph/MovieGraphRenderLayerSubsystem.cpp
+++ b/Engine/Plugins/MovieScene/MovieRenderPipeline/Source/MovieRenderPipelineCore/Private/Graph/MovieGraphRenderLayerSubsystem.cpp
@@ -7,6 +7,7 @@
 #include "Components/VolumetricCloudComponent.h"
 #include "Components/ExponentialHeightFogComponent.h"
 #include "Engine/DirectionalLight.h"
+#include "Engine/Engine.h"
 #include "Engine/PointLight.h"
 #include "Engine/RectLight.h"
 #include "Engine/SkyLight.h"
@@ -644,6 +645,38 @@ bool UMovieGraphConditionGroupQueryBase::IsFirstConditionGroupQuery() const
 	return false;
 }
 
+/**
+ * Finds the counterpart of InSourceActor in InTargetWorld by matching the correct sublevel via
+ * outermost package name, then looking up the actor by name within that level.
+ *
+ * This works more reliably than EditorUtilities::GetSimWorldCounterpartActor / GetEditorWorldCounterpartActor
+ * which can return actors from the wrong sublevel when actor names collide across sublevels.
+ */
+static AActor* FindCounterpartActorInWorld(const AActor* InSourceActor, const UWorld* InTargetWorld)
+{
+	if (!InSourceActor || !InTargetWorld)
+	{
+		return nullptr;
+	}
+
+	// RemovePIEPrefix is a no-op if the name has no PIE prefix, so this works for both directions.
+	const FString SourceLevelPackageName = UWorld::RemovePIEPrefix(InSourceActor->GetLevel()->GetOutermost()->GetName());
+
+	for (FConstLevelIterator LevelIt(InTargetWorld->GetLevelIterator()); LevelIt; ++LevelIt)
+	{
+		if (ULevel* TargetLevel = *LevelIt)
+		{
+			const FString TargetLevelPackageName = UWorld::RemovePIEPrefix(TargetLevel->GetOutermost()->GetName());
+			if (TargetLevelPackageName == SourceLevelPackageName)
+			{
+				return FindObject<AActor>(TargetLevel, *InSourceActor->GetFName().ToString());
+			}
+		}
+	}
+
+	return nullptr;
+}
+
 AActor* UMovieGraphConditionGroupQueryBase::GetActorForCurrentWorld(AActor* InActorToConvert)
 {
 #if WITH_EDITOR
@@ -667,17 +700,14 @@ AActor* UMovieGraphConditionGroupQueryBase::GetActorForCurrentWorld(AActor* InAc
 		if (!bIsEditorActor)
 		{
 #if WITH_EDITOR
-			if (AActor* EditorActor = EditorUtilities::GetEditorWorldCounterpartActor(InActorToConvert))
-			{
-				return EditorActor;
-			}
-#endif
-		}
-		else
-		{
-			// Just use InActorToConvert as-is if it's not from PIE
+			return FindCounterpartActorInWorld(InActorToConvert, GEditor->EditorWorld);
+#else
 			return InActorToConvert;
+#endif
 		}
+
+		// Just use InActorToConvert as-is if it's not from PIE
+		return InActorToConvert;
 	}
 
 	// If a PIE session IS active, try to get the PIE equivalent of the editor actor
@@ -687,20 +717,15 @@ AActor* UMovieGraphConditionGroupQueryBase::GetActorForCurrentWorld(AActor* InAc
 		if (bIsEditorActor)
 		{
 #if WITH_EDITOR
-			if (AActor* PieActor = EditorUtilities::GetSimWorldCounterpartActor(InActorToConvert))
-			{
-				return PieActor;
-			}
+			return FindCounterpartActorInWorld(InActorToConvert, GEditor->PlayWorld);
+#else
+			return FindCounterpartActorInWorld(InActorToConvert, GEngine->GetCurrentPlayWorld());
 #endif
 		}
-		else
-		{
-			// Just use InActorToConvert as-is if it's not from an editor actor
-			return InActorToConvert;
-		}	
-	}
 
-	return nullptr;
+		// Just use InActorToConvert as-is if it's not from an editor actor
+		return InActorToConvert;
+	}
 }
 
 bool FMovieGraphActorQueryEntry::operator==(const FMovieGraphActorQueryEntry& InOther) const
