MoviePipeline: Added a new GetCurrentPlayWorld() helper function to the Movie Graph blueprint function library to assist in getting a world context within blueprints which are run via the Execute Script node.

#Synced-CL 51651671

diff --git a/Engine/Plugins/MovieScene/MovieRenderPipeline/Source/MovieRenderPipelineCore/Private/Graph/MovieGraphBlueprintLibrary.cpp b/Engine/Plugins/MovieScene/MovieRenderPipeline/Source/MovieRenderPipelineCore/Private/Graph/MovieGraphBlueprintLibrary.cpp
--- a/Engine/Plugins/MovieScene/MovieRenderPipeline/Source/MovieRenderPipelineCore/Private/Graph/MovieGraphBlueprintLibrary.cpp
+++ b/Engine/Plugins/MovieScene/MovieRenderPipeline/Source/MovieRenderPipelineCore/Private/Graph/MovieGraphBlueprintLibrary.cpp
@@ -1020,3 +1020,34 @@ UMoviePipelineExecutorShot* UMovieGraphBlueprintLibrary::GetCurrentExecutorShot(
 
 	return nullptr;
 }
+
+UWorld* UMovieGraphBlueprintLibrary::GetCurrentPlayWorld()
+{
+	if (!GEngine)
+	{
+		return nullptr;
+	}
+
+	// Try the standard path first.
+	if (UWorld* PlayWorld = GEngine->GetCurrentPlayWorld())
+	{
+		return PlayWorld;
+	}
+
+	// Fallback: search engine world contexts for a PIE or Game world.
+	// This covers early MRQ callbacks (e.g., OnJobStart) where
+	// GetCurrentPlayWorld() may return null because the world context
+	// hasn't been fully registered yet.
+	for (const FWorldContext& Context : GEngine->GetWorldContexts())
+	{
+		if (Context.WorldType == EWorldType::PIE || Context.WorldType == EWorldType::Game)
+		{
+			if (Context.World())
+			{
+				return Context.World();
+			}
+		}
+	}
+
+	return nullptr;
+}
diff --git a/Engine/Plugins/MovieScene/MovieRenderPipeline/Source/MovieRenderPipelineCore/Public/Graph/MovieGraphBlueprintLibrary.h b/Engine/Plugins/MovieScene/MovieRenderPipeline/Source/MovieRenderPipelineCore/Public/Graph/MovieGraphBlueprintLibrary.h
--- a/Engine/Plugins/MovieScene/MovieRenderPipeline/Source/MovieRenderPipelineCore/Public/Graph/MovieGraphBlueprintLibrary.h
+++ b/Engine/Plugins/MovieScene/MovieRenderPipeline/Source/MovieRenderPipelineCore/Public/Graph/MovieGraphBlueprintLibrary.h
@@ -290,6 +290,13 @@ public:
 	UFUNCTION(BlueprintPure, Category = "Movie Graph")
 	static UE_API UMoviePipelineExecutorShot* GetCurrentExecutorShot(const UMovieGraphPipeline* InMoviePipeline);
 
+	/**
+	 * Returns the current play world (PIE or Game). Useful when a Blueprint does not have a world context object,
+	 * such as in MovieGraphScriptBase callbacks. Returns nullptr if no play world is active.
+	 */
+	UFUNCTION(BlueprintPure, Category = "Movie Graph", meta = (DisplayName = "Get Current Play World"))
+	static UE_API UWorld* GetCurrentPlayWorld();
+
 };
 
 #undef UE_API
