MoviePipeline: Updated post process material (PPM) output in the graph to ensure that alpha=1 when the PPM does not generate alpha.

#Synced-CL 51731107

diff --git a/Engine/Plugins/MovieScene/MovieRenderPipeline/Source/MovieRenderPipelineCore/Public/Graph/MovieGraphDataTypes.h b/Engine/Plugins/MovieScene/MovieRenderPipeline/Source/MovieRenderPipelineCore/Public/Graph/MovieGraphDataTypes.h
--- a/Engine/Plugins/MovieScene/MovieRenderPipeline/Source/MovieRenderPipelineCore/Public/Graph/MovieGraphDataTypes.h
+++ b/Engine/Plugins/MovieScene/MovieRenderPipeline/Source/MovieRenderPipelineCore/Public/Graph/MovieGraphDataTypes.h
@@ -466,6 +466,12 @@ namespace UE::MovieGraph
 		/** Whether this sample represents a sample in the main beauty pass. */
 		bool bIsBeautyPass;
 
+		/**
+		 * When true, forces alpha to 1.0 in the output. Typically used for Post Process Material passes that don't
+		 * explicitly output alpha, where the render target's clear color (alpha=0) would otherwise leak through.
+		 */
+		bool bForceOpaqueAlpha = false;
+
 		/** When using high-res tiling, how many pixels does each tile overlap the adjacent tiles (on each side)? This should be zero if not using tiling. */
 		FIntPoint OverlappedPad;
 
diff --git a/Engine/Plugins/MovieScene/MovieRenderPipeline/Source/MovieRenderPipelineRenderPasses/Private/Graph/Renderers/MovieGraphDeferredPass.cpp b/Engine/Plugins/MovieScene/MovieRenderPipeline/Source/MovieRenderPipelineRenderPasses/Private/Graph/Renderers/MovieGraphDeferredPass.cpp
--- a/Engine/Plugins/MovieScene/MovieRenderPipeline/Source/MovieRenderPipelineRenderPasses/Private/Graph/Renderers/MovieGraphDeferredPass.cpp
+++ b/Engine/Plugins/MovieScene/MovieRenderPipeline/Source/MovieRenderPipelineRenderPasses/Private/Graph/Renderers/MovieGraphDeferredPass.cpp
@@ -17,6 +17,7 @@
 #include "Engine/TextureRenderTarget2D.h"
 #include "EngineModule.h"
 #include "LegacyScreenPercentageDriver.h"
+#include "Materials/Material.h"
 #include "Materials/MaterialInterface.h"
 #include "OpenColorIODisplayExtension.h"
 #include "SceneManagement.h"
@@ -605,6 +606,12 @@ void FMovieGraphDeferredPass::Render(const FMovieGraphTraversalContext& InFrameT
 						// PPM passes are never the beauty pass.
 						PassSampleState.bIsBeautyPass = false;
 
+						// Post process materials that don't explicitly output alpha will inherit the render
+						// target's clear color (alpha=0), producing a fully transparent image. Detect this
+						// case so the accumulator/output path can force alpha to 1.0.
+						UMaterial* BaseMaterial = Material->GetMaterial();
+						PassSampleState.bForceOpaqueAlpha = !(BaseMaterial && BaseMaterial->IsPostProcessMaterialOutputingAlpha());
+
 						// PPM passes may override the file name format that they use.
 						PassSampleState.FilenameFormatOverride = ParentNode->GetPPMFileNameFormatString();
 
diff --git a/Engine/Plugins/MovieScene/MovieRenderPipeline/Source/MovieRenderPipelineRenderPasses/Private/Graph/Renderers/MovieGraphImagePassBase.cpp b/Engine/Plugins/MovieScene/MovieRenderPipeline/Source/MovieRenderPipelineRenderPasses/Private/Graph/Renderers/MovieGraphImagePassBase.cpp
--- a/Engine/Plugins/MovieScene/MovieRenderPipeline/Source/MovieRenderPipelineRenderPasses/Private/Graph/Renderers/MovieGraphImagePassBase.cpp
+++ b/Engine/Plugins/MovieScene/MovieRenderPipeline/Source/MovieRenderPipelineRenderPasses/Private/Graph/Renderers/MovieGraphImagePassBase.cpp
@@ -669,6 +669,14 @@ void AccumulateSample_TaskThread(TUniquePtr<FImagePixelData>&& InPixelData, cons
 	// This reduces memory requirements and improves performance in the baseline case.
 	if (!SampleStatePayload->bRequiresAccumulator)
 	{
+		// PPM passes that don't output alpha arrive with alpha=0 from the render target clear color.
+		// Fix up alpha here since we can't intercept the GPU readback (it's handled by the engine's
+		// buffer visualization system, not the movie pipeline's surface reader).
+		if (SampleStatePayload->bForceOpaqueAlpha)
+		{
+			SamplePixelData->SetAlphaOpaque();
+		}
+
 		OutputMergerPin->OnCompleteRenderPassDataAvailable_AnyThread(MoveTemp(SamplePixelData));
 		return;
 	}
@@ -677,7 +685,12 @@ void AccumulateSample_TaskThread(TUniquePtr<FImagePixelData>&& InPixelData, cons
 	if (AccumulatorPin->NumChannels == 0)
 	{
 		LLM_SCOPE_BYNAME(TEXT("MoviePipeline/ImageAccumulatorInitMemory"));
-		const int32 ChannelCount = 4;
+
+		// Use 3 channels (RGB) when forcing opaque alpha. FetchFullImageValue() initializes
+		// Rgba[3]=1.0 and only overwrites channels 0..NumChannels-1, so a 3-channel accumulator
+		// produces alpha=1.0 for free during the existing fetch pass -- no extra pixel iteration.
+		const int32 ChannelCount = SampleStatePayload->bForceOpaqueAlpha ? 3 : 4;
+
 		AccumulatorPin->InitMemory(SampleStatePayload->AccumulatorResolution, ChannelCount);
 		AccumulatorPin->ZeroPlanes();
 		AccumulatorPin->AccumulationGamma = 1.f;
