MoviePipeline: Fixed a bug that caused lights and decal components to not hide properly via the Modifier node when "Is Hidden" is used.

#Synced-CL 51731107

diff --git a/Engine/Plugins/MovieScene/MovieRenderPipeline/Source/MovieRenderPipelineCore/Private/Graph/MovieGraphRenderLayerSubsystem.cpp b/Engine/Plugins/MovieScene/MovieRenderPipeline/Source/MovieRenderPipelineCore/Private/Graph/MovieGraphRenderLayerSubsystem.cpp
--- a/Engine/Plugins/MovieScene/MovieRenderPipeline/Source/MovieRenderPipelineCore/Private/Graph/MovieGraphRenderLayerSubsystem.cpp
+++ b/Engine/Plugins/MovieScene/MovieRenderPipeline/Source/MovieRenderPipelineCore/Private/Graph/MovieGraphRenderLayerSubsystem.cpp
@@ -258,6 +258,12 @@ void UMovieGraphRenderPropertyModifier::ApplyModifier(const UWorld* World)
 			CachedState.bHoldout = AsExponentialHeightFogComponent->bHoldout;
 			CachedState.bAffectIndirectLightingWhileHidden = !AsExponentialHeightFogComponent->bRenderInMainPass;
 		}
+		else
+		{
+			// General fallback for other USceneComponent types (e.g., lights, decals) that don't inherit
+			// from UPrimitiveComponent. Only bIsHidden is applicable for these component types.
+			CachedState.bIsHidden = InComponent->bHiddenInGame;
+		}
 	};
 
 	auto CacheActorState = [this](AActor* InActor)
@@ -415,6 +421,15 @@ void UMovieGraphRenderPropertyModifier::SetActorAndComponentState(const TArray<F
 		{
 			SetStateForVolumetrics(AsExponentialHeightFogComponent, InNewComponentState);
 		}
+		else if (USceneComponent* AsSceneComponent = Cast<USceneComponent>(InComponent))
+		{
+			// General fallback for other USceneComponent types (e.g., lights, decals) that don't inherit
+			// from UPrimitiveComponent. Only bIsHidden is applicable for these component types.
+			if (bOverride_bIsHidden)
+			{
+				AsSceneComponent->SetHiddenInGame(InNewComponentState.bIsHidden);
+			}
+		}
 	};
 
 	// Sets the state for an actor to the specified actor state
