MoviePipeline: Introduced the scaffolding for a new queue panel, which will enable rendering via the graph, but with a simplified presets-like interface. This is currently deliberately hidden behind a cvar (MovieRenderPipeline.ShowNewQueuePanel) while it is under development. The old queue panel will continue to be shown to users until the the new one is finished.

#Synced-CL 51731107

diff --git a/Engine/Plugins/MovieScene/MovieRenderPipeline/Source/MovieRenderPipelineEditor/Private/Widgets/MoviePipelineQueuePanelUtils.cpp b/Engine/Plugins/MovieScene/MovieRenderPipeline/Source/MovieRenderPipelineEditor/Private/Widgets/MoviePipelineQueuePanelUtils.cpp
new file mode 100644
--- /dev/null
+++ b/Engine/Plugins/MovieScene/MovieRenderPipeline/Source/MovieRenderPipelineEditor/Private/Widgets/MoviePipelineQueuePanelUtils.cpp
@@ -0,0 +1,229 @@
+// Copyright Epic Games, Inc. All Rights Reserved.
+
+#include "MoviePipelineQueuePanelUtils.h"
+
+#include "Editor.h"
+#include "Graph/MovieGraphConfig.h"
+#include "Graph/MovieGraphConfigFactory.h"
+#include "MoviePipelineQueue.h"
+#include "MoviePipelineQueueSubsystem.h"
+#include "MovieRenderPipelineSettings.h"
+
+#include "ContentBrowserModule.h"
+#include "IContentBrowserSingleton.h"
+
+#include "AssetRegistry/AssetData.h"
+#include "AssetRegistry/AssetRegistryModule.h"
+#include "AssetToolsModule.h"
+#include "FileHelpers.h"
+#include "Misc/FileHelper.h"
+#include "Misc/PackageName.h"
+
+#define LOCTEXT_NAMESPACE "MoviePipelineQueuePanelUtils"
+
+namespace UE::MovieRenderPipelineEditor::Private
+{
+	UMovieGraphConfig* GenerateNewShotSubgraph(const UMoviePipelineExecutorJob* InJob, UMoviePipelineExecutorShot* InShot)
+	{
+		UMovieGraphConfigFactory* GraphFactory = NewObject<UMovieGraphConfigFactory>();
+		GraphFactory->InitialSubgraphAsset = InJob->GetGraphPreset();
+
+		const FAssetToolsModule& AssetToolsModule = FAssetToolsModule::GetModule();
+		UObject* NewAsset = AssetToolsModule.Get().CreateAssetWithDialog(GraphFactory->GetSupportedClass(), GraphFactory);
+
+		// A cancel in the dialog returns null -- don't ensure
+		UMovieGraphConfig* NewGraph = Cast<UMovieGraphConfig>(NewAsset);
+		if (!NewGraph)
+		{
+			return nullptr;
+		}
+
+		constexpr bool bOnlyDirty = false;
+		UEditorLoadingAndSavingUtils::SavePackages({ NewGraph->GetPackage() }, bOnlyDirty);
+
+		InShot->SetGraphPreset(NewGraph);
+		return NewGraph;
+	}
+
+	void SaveTransientQueueToAsset(UMoviePipelineQueue* DestinationQueue)
+	{
+		if (!DestinationQueue)
+		{
+			return;
+		}
+
+		UMoviePipelineQueueSubsystem* Subsystem = GEditor->GetEditorSubsystem<UMoviePipelineQueueSubsystem>();
+		check(Subsystem);
+
+		// Copy the live (transient) queue state into the destination asset
+		UMoviePipelineQueue* CurrentQueue = Subsystem->GetQueue();
+		DestinationQueue->CopyFrom(CurrentQueue);
+		DestinationQueue->SetQueueOrigin(nullptr);
+		DestinationQueue->MarkPackageDirty();
+		DestinationQueue->SetFlags(RF_Public | RF_Standalone | RF_Transactional);
+
+		FAssetRegistryModule::AssetCreated(DestinationQueue);
+
+		constexpr bool bCheckDirty = false;
+		constexpr bool bPromptToSave = false;
+		const FEditorFileUtils::EPromptReturnCode PromptReturnCode =
+			FEditorFileUtils::PromptForCheckoutAndSave({ DestinationQueue->GetPackage() }, bCheckDirty, bPromptToSave);
+
+		if (PromptReturnCode == FEditorFileUtils::EPromptReturnCode::PR_Success)
+		{
+			// Reload the queue from the newly saved asset so the subsystem tracks it as the origin
+			constexpr bool bPromptOnReplacingDirtyQueue = false;
+			Subsystem->LoadQueue(DestinationQueue, bPromptOnReplacingDirtyQueue);
+		}
+	}
+
+	void SaveQueue()
+	{
+		UMoviePipelineQueue* QueueOrigin = GetQueueOrigin();
+
+		// If the queue has no known origin asset, fall through to "Save As" so the user picks a path
+		if (!QueueOrigin)
+		{
+			SaveQueueAs();
+			return;
+		}
+
+		SaveTransientQueueToAsset(QueueOrigin);
+	}
+
+	void SaveQueueAs()
+	{
+		UMoviePipelineQueueSubsystem* Subsystem = GEditor->GetEditorSubsystem<UMoviePipelineQueueSubsystem>();
+		check(Subsystem);
+		UMoviePipelineQueue* CurrentQueue = Subsystem->GetQueue();
+
+		FString PackageName;
+		if (!GetSavePresetPackageName(CurrentQueue->GetName(), PackageName))
+		{
+			return;
+		}
+
+		const FString NewAssetName = FPackageName::GetLongPackageAssetName(PackageName);
+		UPackage* NewPackage = CreatePackage(*PackageName);
+		NewPackage->MarkAsFullyLoaded();
+
+		UMoviePipelineQueue* DuplicateQueue = DuplicateObject<UMoviePipelineQueue>(CurrentQueue, NewPackage, *NewAssetName);
+		SaveTransientQueueToAsset(DuplicateQueue);
+	}
+
+	void ImportSavedQueueAsset(const FAssetData& InPresetAsset)
+	{
+		UMoviePipelineQueueSubsystem* Subsystem = GEditor->GetEditorSubsystem<UMoviePipelineQueueSubsystem>();
+		check(Subsystem);
+
+		UMoviePipelineQueue* SavedQueue = CastChecked<UMoviePipelineQueue>(InPresetAsset.GetAsset());
+		Subsystem->LoadQueue(SavedQueue);
+	}
+
+	bool IsQueueDirty()
+	{
+		const UMoviePipelineQueueSubsystem* Subsystem = GEditor->GetEditorSubsystem<UMoviePipelineQueueSubsystem>();
+		check(Subsystem);
+		return Subsystem->IsQueueDirty();
+	}
+
+	UMoviePipelineQueue* GetQueueOrigin()
+	{
+		UMoviePipelineQueueSubsystem* Subsystem = GEditor->GetEditorSubsystem<UMoviePipelineQueueSubsystem>();
+		check(Subsystem);
+
+		const UMoviePipelineQueue* Queue = Subsystem->GetQueue();
+		return Queue ? Queue->GetQueueOrigin() : nullptr;
+	}
+
+	FString GetQueueOriginName()
+	{
+		const UMoviePipelineQueue* QueueOrigin = GetQueueOrigin();
+		if (!QueueOrigin)
+		{
+			return FString();
+		}
+
+		const UPackage* Package = QueueOrigin->GetPackage();
+		if (!Package)
+		{
+			return FString();
+		}
+
+		FString OutPackageRoot, OutPackagePath, OutPackageName;
+		if (FPackageName::SplitLongPackageName(Package->GetName(), OutPackageRoot, OutPackagePath, OutPackageName))
+		{
+			return OutPackageName;
+		}
+
+		return FString();
+	}
+
+	bool OpenSaveDialog(const FString& InDefaultPath, const FString& InNewNameSuggestion, FString& OutPackageName)
+	{
+		FSaveAssetDialogConfig SaveAssetDialogConfig;
+		SaveAssetDialogConfig.DefaultPath = InDefaultPath;
+		SaveAssetDialogConfig.DefaultAssetName = InNewNameSuggestion;
+		SaveAssetDialogConfig.AssetClassNames.Add(UMoviePipelineQueue::StaticClass()->GetClassPathName());
+		SaveAssetDialogConfig.ExistingAssetPolicy = ESaveAssetDialogExistingAssetPolicy::AllowButWarn;
+		SaveAssetDialogConfig.DialogTitleOverride = LOCTEXT("SaveQueueAssetDialogTitle", "Save Queue Asset");
+
+		FContentBrowserModule& ContentBrowserModule = FModuleManager::LoadModuleChecked<FContentBrowserModule>("ContentBrowser");
+		const FString SaveObjectPath = ContentBrowserModule.Get().CreateModalSaveAssetDialog(SaveAssetDialogConfig);
+
+		if (!SaveObjectPath.IsEmpty())
+		{
+			OutPackageName = FPackageName::ObjectPathToPackageName(SaveObjectPath);
+			return true;
+		}
+
+		return false;
+	}
+
+	bool GetSavePresetPackageName(const FString& InExistingName, FString& OutName)
+	{
+		UMovieRenderPipelineProjectSettings* ConfigSettings = GetMutableDefault<UMovieRenderPipelineProjectSettings>();
+
+		const FString DefaultSaveDirectory = ConfigSettings->PresetSaveDir.Path;
+
+		FString DialogStartPath;
+		FPackageName::TryConvertFilenameToLongPackageName(DefaultSaveDirectory, DialogStartPath);
+		if (DialogStartPath.IsEmpty())
+		{
+			DialogStartPath = TEXT("/Game");
+		}
+
+		FString UniquePackageName;
+		FString UniqueAssetName;
+
+		FAssetToolsModule& AssetToolsModule = FModuleManager::LoadModuleChecked<FAssetToolsModule>("AssetTools");
+		AssetToolsModule.Get().CreateUniqueAssetName(DialogStartPath / InExistingName, TEXT(""), UniquePackageName, UniqueAssetName);
+
+		const FString DialogStartName = FPaths::GetCleanFilename(UniqueAssetName);
+
+		FString UserPackageName;
+		FString NewPackageName;
+
+		bool bFilenameValid = false;
+		while (!bFilenameValid)
+		{
+			if (!OpenSaveDialog(DialogStartPath, DialogStartName, UserPackageName))
+			{
+				return false;
+			}
+
+			NewPackageName = UserPackageName;
+
+			FText OutError;
+			bFilenameValid = FFileHelper::IsFilenameValidForSaving(NewPackageName, OutError);
+		}
+
+		// Remember the last save location for next time
+		ConfigSettings->PresetSaveDir.Path = FPackageName::GetLongPackagePath(UserPackageName);
+		ConfigSettings->SaveConfig();
+		OutName = MoveTemp(NewPackageName);
+		return true;
+	}
+}
+
+#undef LOCTEXT_NAMESPACE // MoviePipelineQueuePanelUtils
diff --git a/Engine/Plugins/MovieScene/MovieRenderPipeline/Source/MovieRenderPipelineEditor/Private/Widgets/MoviePipelineQueuePanelUtils.h b/Engine/Plugins/MovieScene/MovieRenderPipeline/Source/MovieRenderPipelineEditor/Private/Widgets/MoviePipelineQueuePanelUtils.h
new file mode 100644
--- /dev/null
+++ b/Engine/Plugins/MovieScene/MovieRenderPipeline/Source/MovieRenderPipelineEditor/Private/Widgets/MoviePipelineQueuePanelUtils.h
@@ -0,0 +1,47 @@
+// Copyright Epic Games, Inc. All Rights Reserved.
+
+#pragma once
+
+#include "CoreMinimal.h"
+
+class UMoviePipelineExecutorJob;
+class UMoviePipelineExecutorShot;
+class UMoviePipelineQueue;
+class UMovieGraphConfig;
+struct FAssetData;
+
+namespace UE::MovieRenderPipelineEditor::Private
+{
+	/**
+	 * Generates and returns a new shot subgraph based on the given job and shot. Returns nullptr if there was an issue creating the graph. If the
+	 * graph was successfully created, it will be assigned to the shot.
+	 */
+	UMovieGraphConfig* GenerateNewShotSubgraph(const UMoviePipelineExecutorJob* InJob, UMoviePipelineExecutorShot* InShot);
+
+	/** Copies the transient queue to the destination queue asset, registers it, and prompts for save. */
+	void SaveTransientQueueToAsset(UMoviePipelineQueue* DestinationQueue);
+
+	/** Saves the queue. If it has no known origin, falls through to SaveQueueAs(). */
+	void SaveQueue();
+
+	/** Prompts the user for a new package name and saves the queue as a new asset. */
+	void SaveQueueAs();
+
+	/** Loads a previously saved queue asset via the subsystem. Does NOT dismiss menus -- callers handle that. */
+	void ImportSavedQueueAsset(const FAssetData& InPresetAsset);
+
+	/** Returns true if the current queue has been modified since the last save. */
+	bool IsQueueDirty();
+
+	/** Returns the origin queue asset for the currently loaded queue, or nullptr. */
+	UMoviePipelineQueue* GetQueueOrigin();
+
+	/** Returns the short package name of the queue origin, or an empty string. */
+	FString GetQueueOriginName();
+
+	/** Opens a save-asset dialog for a queue. Returns true if the user selected a valid path. */
+	bool OpenSaveDialog(const FString& InDefaultPath, const FString& InNewNameSuggestion, FString& OutPackageName);
+
+	/** Determines the package name when saving a queue preset. Returns true on success. */
+	bool GetSavePresetPackageName(const FString& InExistingName, FString& OutName);
+}
diff --git a/Engine/Plugins/MovieScene/MovieRenderPipeline/Source/MovieRenderPipelineEditor/Private/Widgets/SMoviePipelineNewQueuePanel.cpp b/Engine/Plugins/MovieScene/MovieRenderPipeline/Source/MovieRenderPipelineEditor/Private/Widgets/SMoviePipelineNewQueuePanel.cpp
new file mode 100644
--- /dev/null
+++ b/Engine/Plugins/MovieScene/MovieRenderPipeline/Source/MovieRenderPipelineEditor/Private/Widgets/SMoviePipelineNewQueuePanel.cpp
@@ -0,0 +1,682 @@
+// Copyright Epic Games, Inc. All Rights Reserved.
+
+#include "Widgets/SMoviePipelineNewQueuePanel.h"
+#include "MoviePipelineQueuePanelUtils.h"
+#include "Customizations/Graph/MovieGraphNamedResolutionCustomization.h"
+#include "Customizations/JobCustomization.h"
+#include "Widgets/MoviePipelineWidgetConstants.h"
+#include "SMoviePipelineQueueEditor.h"
+#include "SMoviePipelineConfigPanel.h"
+#include "MovieRenderPipelineEditorUtils.h"
+#include "MovieRenderPipelineSettings.h"
+#include "MovieRenderPipelineStyle.h"
+#include "MoviePipelineQueue.h"
+#include "MoviePipelinePrimaryConfig.h"
+#include "MoviePipelineQueueSubsystem.h"
+#include "MoviePipelineTelemetry.h"
+#include "Graph/MovieGraphConfig.h"
+#include "Graph/MovieGraphAssetToolkit.h"
+
+#include "Widgets/SBoxPanel.h"
+#include "Widgets/Input/SComboBox.h"
+#include "Widgets/Layout/SBorder.h"
+#include "Widgets/Text/STextBlock.h"
+#include "Widgets/Input/SButton.h"
+#include "Widgets/SWindow.h"
+#include "Styling/AppStyle.h"
+#include "Framework/Application/SlateApplication.h"
+#include "Widgets/Layout/SSplitter.h"
+#include "Widgets/Layout/SWidgetSwitcher.h"
+#include "Widgets/Images/SImage.h"
+#include "Framework/MultiBox/MultiBoxBuilder.h"
+#include "SPositiveActionButton.h"
+#include "Styling/ToolBarStyle.h"
+
+#include "IContentBrowserSingleton.h"
+#include "ContentBrowserModule.h"
+
+#include "Editor.h"
+#include "PropertyEditorModule.h"
+#include "IDetailsView.h"
+#include "ISettingsModule.h"
+#include "ISinglePropertyView.h"
+#include "Modules/ModuleManager.h"
+#include "AssetRegistry/AssetData.h"
+
+#define LOCTEXT_NAMESPACE "SMoviePipelineNewQueuePanel"
+
+void SMoviePipelineNewQueuePanel::Construct(const FArguments& InArgs)
+{
+	FPropertyEditorModule& PropertyEditorModule = FModuleManager::Get().LoadModuleChecked<FPropertyEditorModule>("PropertyEditor");
+
+	FDetailsViewArgs DetailsViewArgs;
+	DetailsViewArgs.bAllowSearch = false;
+	DetailsViewArgs.NameAreaSettings = FDetailsViewArgs::HideNameArea;
+	DetailsViewArgs.bHideSelectionTip = true;
+	DetailsViewArgs.ColumnWidth = 0.7f;
+
+	JobDetailsPanelWidget = PropertyEditorModule.CreateDetailView(DetailsViewArgs);
+
+	JobDetailsPanelWidget->RegisterInstancedCustomPropertyLayout(
+		UMoviePipelineExecutorJob::StaticClass(),
+		FOnGetDetailCustomizationInstance::CreateStatic(&FJobDetailsCustomization::MakeInstance));
+
+	JobDetailsPanelWidget->RegisterInstancedCustomPropertyLayout(
+		UMoviePipelineExecutorShot::StaticClass(),
+		FOnGetDetailCustomizationInstance::CreateStatic(&FJobDetailsCustomization::MakeInstance));
+
+	JobDetailsPanelWidget->RegisterInstancedCustomPropertyTypeLayout(
+		FMovieGraphNamedResolution::StaticStruct()->GetFName(),
+		FOnGetPropertyTypeCustomizationInstance::CreateStatic(&FMovieGraphNamedResolutionCustomization::MakeInstance));
+
+	PipelineQueueEditorWidget = SNew(SMoviePipelineQueueEditor)
+		.OnEditConfigRequested(this, &SMoviePipelineNewQueuePanel::OnEditJobConfigRequested)
+		.OnPresetChosen(this, &SMoviePipelineNewQueuePanel::OnJobPresetChosen)
+		.OnJobSelectionChanged(this, &SMoviePipelineNewQueuePanel::OnSelectionChanged);
+
+	const UMoviePipelineQueueSubsystem* Subsystem = GEditor->GetEditorSubsystem<UMoviePipelineQueueSubsystem>();
+	check(Subsystem);
+	UMoviePipelineQueue* Queue = Subsystem->GetQueue();
+	check(Queue);
+
+	// Auto-select the first job so the details panel isn't empty on open
+	{
+		const TArray<UMoviePipelineExecutorJob*> CurrentJobs = Queue->GetJobs();
+		if (!CurrentJobs.IsEmpty())
+		{
+			PipelineQueueEditorWidget->SetSelectedJobs({ CurrentJobs[0] });
+		}
+	}
+
+	// "After Render" behavior dropdown
+	TSharedPtr<ISinglePropertyView> PostRenderBehaviorWidget;
+	{
+		UMoviePipelineQueueProjectSettings* QueueProjectSettings = GetMutableDefault<UMoviePipelineQueueProjectSettings>();
+
+		FSinglePropertyParams PostRenderActionTypeParams;
+		PostRenderActionTypeParams.NamePlacement = EPropertyNamePlacement::Hidden;
+		PostRenderActionTypeParams.bHideResetToDefault = true;
+
+		PostRenderBehaviorWidget = PropertyEditorModule.CreateSingleProperty(
+			QueueProjectSettings, GET_MEMBER_NAME_CHECKED(UMoviePipelineQueueProjectSettings, PostRenderActionType), PostRenderActionTypeParams);
+	}
+
+	ChildSlot
+	[
+		SNew(SVerticalBox)
+
+		// Toolbar
+		+ SVerticalBox::Slot()
+		.Padding(FMargin(0.f, 1.0f))
+		.AutoHeight()
+		[
+			SNew(SBorder)
+			.BorderImage(FAppStyle::GetBrush("DetailsView.CategoryTop"))
+			.BorderBackgroundColor(FLinearColor(.6, .6, .6, 1.0f))
+			[
+				SNew(SHorizontalBox)
+
+				// Add a Level Sequence to the queue
+				+ SHorizontalBox::Slot()
+				.Padding(MoviePipeline::ButtonOffset)
+				.VAlign(VAlign_Fill)
+				.AutoWidth()
+				[
+					SNew(SPositiveActionButton)
+					.OnGetMenuContent_Lambda([this]() -> TSharedRef<SWidget>
+					{
+						return PipelineQueueEditorWidget->OnGenerateNewJobFromAssetMenu();
+					})
+					.Icon(FAppStyle::Get().GetBrush("Icons.Plus"))
+					.Text(LOCTEXT("AddNewJob_Text", "Add"))
+				]
+
+				// Render queue (local/remote dropdown)
+				+ SHorizontalBox::Slot()
+				.Padding(MoviePipeline::ButtonOffset)
+				.VAlign(VAlign_Fill)
+				.AutoWidth()
+				[
+					SNew(SActionButton)
+					.ActionButtonType(EActionButtonType::Primary)
+					.HasDownArrow(true)
+					.Icon(FAppStyle::Get().GetBrush("LevelEditor.OpenCinematic"))
+					.Text(LOCTEXT("RenderQueue_Text", "Render"))
+					.IsEnabled(this, &SMoviePipelineNewQueuePanel::IsRenderEnabled)
+					.OnGetMenuContent(this, &SMoviePipelineNewQueuePanel::OnGenerateRenderMenu)
+				]
+
+				// Spacer
+				+ SHorizontalBox::Slot()
+				.VAlign(VAlign_Fill)
+				.HAlign(HAlign_Fill)
+				.FillWidth(1.f)
+				[
+					SNullWidget::NullWidget
+				]
+
+				// Queue save/load dropdown
+				+ SHorizontalBox::Slot()
+				.Padding(MoviePipeline::ButtonOffset)
+				.VAlign(VAlign_Center)
+				.HAlign(HAlign_Right)
+				.AutoWidth()
+				[
+					SNew(SComboButton)
+					.ToolTipText(LOCTEXT("QueueManagementButton_Tooltip", "Export the current queue to an asset, or load a previously saved queue."))
+					.ContentPadding(MoviePipeline::ButtonPadding)
+					.ComboButtonStyle(FMovieRenderPipelineStyle::Get(), "ComboButton")
+					.OnGetMenuContent(this, &SMoviePipelineNewQueuePanel::OnGenerateSavedQueuesMenu)
+					.ForegroundColor(FSlateColor::UseForeground())
+					.ButtonContent()
+					[
+						SNew(SHorizontalBox)
+						+ SHorizontalBox::Slot()
+						.Padding(0, 1, 4, 0)
+						.AutoWidth()
+						[
+							SNew(SImage)
+							.Image(FAppStyle::Get().GetBrush("AssetEditor.SaveAsset"))
+						]
+
+						+ SHorizontalBox::Slot()
+						.Padding(0, 1, 0, 0)
+						[
+							SNew(STextBlock)
+							.Text(this, &SMoviePipelineNewQueuePanel::GetQueueMenuButtonText)
+						]
+					]
+				]
+			]
+		]
+
+		// Main queue body
+		+ SVerticalBox::Slot()
+		.FillHeight(1.0f)
+		[
+			SNew(SSplitter)
+			.Orientation(EOrientation::Orient_Horizontal)
+			+ SSplitter::Slot()
+			.Value(3)
+			[
+				PipelineQueueEditorWidget.ToSharedRef()
+			]
+			+ SSplitter::Slot()
+			.Value(1)
+			[
+				SNew(SBorder)
+				.BorderImage(FAppStyle::GetBrush("ToolPanel.GroupBorder"))
+				.Padding(FMargin(1.f, 1.0f))
+				.Content()
+				[
+					SNew(SWidgetSwitcher)
+					.WidgetIndex(this, &SMoviePipelineNewQueuePanel::GetDetailsViewWidgetIndex)
+					.IsEnabled(this, &SMoviePipelineNewQueuePanel::IsDetailsViewEnabled)
+					+ SWidgetSwitcher::Slot()
+					[
+						JobDetailsPanelWidget.ToSharedRef()
+					]
+					+ SWidgetSwitcher::Slot()
+					.Padding(2.0f, 24.0f, 2.0f, 2.0f)
+					[
+						SNew(SBox)
+						.HAlign(EHorizontalAlignment::HAlign_Center)
+						.Content()
+						[
+							SNew(STextBlock)
+							.Text(LOCTEXT("NoJobSelected", "Select a job to view details."))
+						]
+					]
+				]
+			]
+		]
+
+		// Footer
+		+ SVerticalBox::Slot()
+		.AutoHeight()
+		[
+			SNew(SBorder)
+			.BorderImage(FAppStyle::GetBrush("DetailsView.CategoryTop"))
+			.BorderBackgroundColor(FLinearColor(.6, .6, .6, 1.0f))
+			.Padding(FMargin(0, 2, 0, 2))
+			[
+				SNew(SHorizontalBox)
+
+				+ SHorizontalBox::Slot()
+				.VAlign(VAlign_Center)
+				.HAlign(HAlign_Left)
+				.Padding(5, 0, 0, 0)
+				.AutoWidth()
+				[
+					SNew(STextBlock)
+					.Text(LOCTEXT("AfterRender_Text", "After Render:"))
+				]
+
+				+ SHorizontalBox::Slot()
+				.VAlign(VAlign_Center)
+				.HAlign(HAlign_Left)
+				.AutoWidth()
+				[
+					PostRenderBehaviorWidget.ToSharedRef()
+				]
+
+				+ SHorizontalBox::Slot()
+				.VAlign(VAlign_Center)
+				.HAlign(HAlign_Left)
+				.AutoWidth()
+				[
+					SNew(SButton)
+					.ButtonStyle(FAppStyle::Get(), TEXT("SimpleButton"))
+					.ContentPadding(0)
+					.OnClicked_Lambda([]()
+					{
+						FModuleManager::LoadModuleChecked<ISettingsModule>("Settings")
+							.ShowViewer("Editor", "Plugins", "MovieRenderGraphEditorSettings");
+
+						return FReply::Handled();
+					})
+					.ToolTipText(LOCTEXT("OpenPlaybackPrefs_Tooltip", "Open up the preferences which dictate how media is opened after a render completes."))
+					[
+						SNew(SImage)
+						.ColorAndOpacity(FSlateColor::UseForeground())
+						.Image(FAppStyle::Get().GetBrush("EditorPreferences.TabIcon"))
+						.DesiredSizeOverride(FVector2D(16.f, 16.f))
+					]
+				]
+
+				+ SHorizontalBox::Slot()
+				.VAlign(VAlign_Fill)
+				.HAlign(HAlign_Left)
+				.FillWidth(1.f)
+				[
+					SNullWidget::NullWidget
+				]
+			]
+		]
+	];
+}
+
+void SMoviePipelineNewQueuePanel::OnRenderLocalRequested()
+{
+	UE::MovieRenderPipelineEditor::Private::PerformLocalRender();
+}
+
+bool SMoviePipelineNewQueuePanel::IsRenderLocalEnabled() const
+{
+	const bool bConfigWindowIsOpen = WeakEditorWindow.IsValid();
+	return UE::MovieRenderPipelineEditor::Private::CanPerformLocalRender() && !bConfigWindowIsOpen;
+}
+
+void SMoviePipelineNewQueuePanel::OnRenderRemoteRequested()
+{
+	UMoviePipelineQueueSubsystem* Subsystem = GEditor->GetEditorSubsystem<UMoviePipelineQueueSubsystem>();
+	check(Subsystem);
+
+	const UMovieRenderPipelineProjectSettings* ProjectSettings = GetDefault<UMovieRenderPipelineProjectSettings>();
+	TSubclassOf<UMoviePipelineExecutorBase> ExecutorClass = ProjectSettings->DefaultRemoteExecutor.TryLoadClass<UMoviePipelineExecutorBase>();
+
+	// OnRenderRemoteRequested should only get called if IsRenderRemoteEnabled() returns true, meaning there's a valid class.
+	check(ExecutorClass != nullptr);
+
+	Subsystem->RenderQueueWithExecutor(ExecutorClass);
+
+	constexpr bool bIsLocal = false;
+	FMoviePipelineTelemetry::SendRendersRequestedTelemetry(bIsLocal, Subsystem->GetQueue()->GetJobs());
+}
+
+bool SMoviePipelineNewQueuePanel::IsRenderRemoteEnabled() const
+{
+	UMoviePipelineQueueSubsystem* Subsystem = GEditor->GetEditorSubsystem<UMoviePipelineQueueSubsystem>();
+	check(Subsystem);
+
+	const UMovieRenderPipelineProjectSettings* ProjectSettings = GetDefault<UMovieRenderPipelineProjectSettings>();
+	const bool bHasExecutor = ProjectSettings->DefaultRemoteExecutor.TryLoadClass<UMoviePipelineExecutorBase>() != nullptr;
+	const bool bNotRendering = !Subsystem->IsRendering();
+	const bool bConfigWindowIsOpen = WeakEditorWindow.IsValid();
+
+	bool bAtLeastOneJobAvailable = false;
+	for (const UMoviePipelineExecutorJob* Job : Subsystem->GetQueue()->GetJobs())
+	{
+		if (!Job->IsConsumed() && Job->IsEnabled())
+		{
+			bAtLeastOneJobAvailable = true;
+			break;
+		}
+	}
+
+	return bHasExecutor && bNotRendering && bAtLeastOneJobAvailable && !bConfigWindowIsOpen;
+}
+
+bool SMoviePipelineNewQueuePanel::IsRenderEnabled() const
+{
+	return IsRenderLocalEnabled() || IsRenderRemoteEnabled();
+}
+
+TSharedRef<SWidget> SMoviePipelineNewQueuePanel::OnGenerateRenderMenu()
+{
+	FMenuBuilder MenuBuilder(true, nullptr);
+
+	MenuBuilder.AddMenuEntry(
+		LOCTEXT("RenderLocal_Text", "Render (Local)"),
+		LOCTEXT("RenderLocal_Tooltip", "Renders the current queue in the current process using Play in Editor."),
+		FSlateIcon(),
+		FUIAction(
+			FExecuteAction::CreateSP(this, &SMoviePipelineNewQueuePanel::OnRenderLocalRequested),
+			FCanExecuteAction::CreateSP(this, &SMoviePipelineNewQueuePanel::IsRenderLocalEnabled)
+		)
+	);
+
+	MenuBuilder.AddMenuEntry(
+		LOCTEXT("RenderRemote_Text", "Render (Remote)"),
+		LOCTEXT("RenderRemote_Tooltip", "Renders the current queue in a separate process."),
+		FSlateIcon(),
+		FUIAction(
+			FExecuteAction::CreateSP(this, &SMoviePipelineNewQueuePanel::OnRenderRemoteRequested),
+			FCanExecuteAction::CreateSP(this, &SMoviePipelineNewQueuePanel::IsRenderRemoteEnabled)
+		)
+	);
+
+	return MenuBuilder.MakeWidget();
+}
+
+void SMoviePipelineNewQueuePanel::OnJobPresetChosen(TWeakObjectPtr<UMoviePipelineExecutorJob> InJob, TWeakObjectPtr<UMoviePipelineExecutorShot> InShot)
+{
+	UMovieRenderPipelineProjectSettings* ProjectSettings = GetMutableDefault<UMovieRenderPipelineProjectSettings>();
+	if (!InShot.IsValid())
+	{
+		ProjectSettings->LastPresetOrigin = InJob->GetPresetOrigin();
+	}
+	ProjectSettings->SaveConfig();
+}
+
+void SMoviePipelineNewQueuePanel::OnEditJobConfigRequested(TWeakObjectPtr<UMoviePipelineExecutorJob> InJob, TWeakObjectPtr<UMoviePipelineExecutorShot> InShot)
+{
+	// Only allow one config editor open at a time
+	if (WeakEditorWindow.IsValid())
+	{
+		FWidgetPath ExistingWindowPath;
+		if (FSlateApplication::Get().FindPathToWidget(WeakEditorWindow.Pin().ToSharedRef(), ExistingWindowPath, EVisibility::All))
+		{
+			WeakEditorWindow.Pin()->BringToFront();
+			FSlateApplication::Get().SetAllUserFocus(ExistingWindowPath, EFocusCause::SetDirectly);
+		}
+		return;
+	}
+
+	// Determine whether to open a graph editor or the legacy config panel
+	UMovieGraphConfig* GraphToEdit = nullptr;
+	if (InShot.IsValid() && InShot->IsUsingGraphConfiguration())
+	{
+		GraphToEdit = InShot->GetGraphPreset();
+		if (!GraphToEdit)
+		{
+			GraphToEdit = UE::MovieRenderPipelineEditor::Private::GenerateNewShotSubgraph(InJob.Get(), InShot.Get());
+			if (!GraphToEdit)
+			{
+				return;
+			}
+		}
+	}
+	else if (InJob.IsValid() && InJob->IsUsingGraphConfiguration())
+	{
+		GraphToEdit = InJob->GetGraphPreset();
+	}
+
+	if (GraphToEdit)
+	{
+		GEditor->GetEditorSubsystem<UAssetEditorSubsystem>()->OpenEditorForAsset(GraphToEdit);
+		return;
+	}
+
+	// Legacy config panel path
+	TSubclassOf<UMoviePipelineConfigBase> ConfigType;
+	UMoviePipelineConfigBase* BasePreset = nullptr;
+	UMoviePipelineConfigBase* BaseConfig = nullptr;
+	if (InShot.IsValid())
+	{
+		ConfigType = UMoviePipelineShotConfig::StaticClass();
+		BasePreset = InShot->GetShotOverridePresetOrigin();
+		BaseConfig = InShot->GetShotOverrideConfiguration();
+	}
+	else
+	{
+		ConfigType = UMoviePipelinePrimaryConfig::StaticClass();
+		BasePreset = InJob->GetPresetOrigin();
+		BaseConfig = InJob->GetConfiguration();
+	}
+
+	TSharedRef<SWindow> EditorWindow =
+		SNew(SWindow)
+		.ClientSize(FVector2D(700, 600));
+
+	TSharedRef<SMoviePipelineConfigPanel> ConfigEditorPanel =
+		SNew(SMoviePipelineConfigPanel, ConfigType)
+		.Job(InJob)
+		.Shot(InShot)
+		.OnConfigurationModified(this, &SMoviePipelineNewQueuePanel::OnConfigUpdatedForJob)
+		.OnConfigurationSetToPreset(this, &SMoviePipelineNewQueuePanel::OnConfigUpdatedForJobToPreset)
+		.BasePreset(BasePreset)
+		.BaseConfig(BaseConfig);
+
+	EditorWindow->SetContent(ConfigEditorPanel);
+
+	TSharedPtr<SWindow> ParentWindow = FSlateApplication::Get().FindWidgetWindow(AsShared());
+	if (ParentWindow.IsValid())
+	{
+		FSlateApplication::Get().AddWindowAsNativeChild(EditorWindow, ParentWindow.ToSharedRef());
+	}
+
+	WeakEditorWindow = EditorWindow;
+}
+
+void SMoviePipelineNewQueuePanel::OnConfigWindowClosed()
+{
+	if (WeakEditorWindow.IsValid())
+	{
+		WeakEditorWindow.Pin()->RequestDestroyWindow();
+	}
+}
+
+void SMoviePipelineNewQueuePanel::OnConfigUpdatedForJob(TWeakObjectPtr<UMoviePipelineExecutorJob> InJob, TWeakObjectPtr<UMoviePipelineExecutorShot> InShot, UMoviePipelineConfigBase* InConfig)
+{
+	if (InJob.IsValid())
+	{
+		if (InShot.IsValid())
+		{
+			if (UMoviePipelineShotConfig* ShotConfig = Cast<UMoviePipelineShotConfig>(InConfig))
+			{
+				InShot->SetShotOverrideConfiguration(ShotConfig);
+			}
+		}
+		else
+		{
+			if (UMoviePipelinePrimaryConfig* PrimaryConfig = Cast<UMoviePipelinePrimaryConfig>(InConfig))
+			{
+				InJob->SetConfiguration(PrimaryConfig);
+			}
+		}
+	}
+
+	OnConfigWindowClosed();
+}
+
+void SMoviePipelineNewQueuePanel::OnConfigUpdatedForJobToPreset(TWeakObjectPtr<UMoviePipelineExecutorJob> InJob, TWeakObjectPtr<UMoviePipelineExecutorShot> InShot, UMoviePipelineConfigBase* InConfig)
+{
+	if (InJob.IsValid())
+	{
+		if (InShot.IsValid())
+		{
+			if (UMoviePipelineShotConfig* ShotConfig = Cast<UMoviePipelineShotConfig>(InConfig))
+			{
+				InShot->SetShotOverridePresetOrigin(ShotConfig);
+			}
+		}
+		else
+		{
+			if (UMoviePipelinePrimaryConfig* PrimaryConfig = Cast<UMoviePipelinePrimaryConfig>(InConfig))
+			{
+				InJob->SetPresetOrigin(PrimaryConfig);
+			}
+		}
+	}
+
+	OnJobPresetChosen(InJob, InShot);
+	OnConfigWindowClosed();
+}
+
+void SMoviePipelineNewQueuePanel::OnSelectionChanged(const TArray<UMoviePipelineExecutorJob*>& InSelectedJobs, const TArray<UMoviePipelineExecutorShot*>& InSelectedShots)
+{
+	// Prefer showing shot details when a shot is selected, otherwise show the job
+	TArray<UObject*> ObjectsToDisplay;
+	if (!InSelectedShots.IsEmpty())
+	{
+		ObjectsToDisplay.Append(InSelectedShots);
+	}
+	else
+	{
+		ObjectsToDisplay.Append(InSelectedJobs);
+	}
+
+	JobDetailsPanelWidget->SetObjects(ObjectsToDisplay);
+	NumSelectedJobs = InSelectedJobs.Num();
+}
+
+int32 SMoviePipelineNewQueuePanel::GetDetailsViewWidgetIndex() const
+{
+	return NumSelectedJobs == 0;
+}
+
+bool SMoviePipelineNewQueuePanel::IsDetailsViewEnabled() const
+{
+	TArray<TWeakObjectPtr<UObject>> OutObjects = JobDetailsPanelWidget->GetSelectedObjects();
+
+	for (const TWeakObjectPtr<UObject>& Object : OutObjects)
+	{
+		const UMoviePipelineExecutorJob* Job = Cast<UMoviePipelineExecutorJob>(Object);
+		if (Job && Job->IsConsumed())
+		{
+			return false;
+		}
+	}
+
+	return true;
+}
+
+TSharedRef<SWidget> SMoviePipelineNewQueuePanel::OnGenerateSavedQueuesMenu()
+{
+	FMenuBuilder MenuBuilder(true, nullptr);
+
+	IContentBrowserSingleton& ContentBrowser = FModuleManager::LoadModuleChecked<FContentBrowserModule>("ContentBrowser").Get();
+
+	// Display the currently loaded queue name (or "None") as a disabled header entry
+	static const FText NoLoadedQueueText = LOCTEXT("NoLoadedQueue", "None");
+	static const FText LoadedQueueFormatText = LOCTEXT("CurrentQueueFormat", "Current Queue: {0}");
+	const FString QueueName = GetQueueOriginName();
+	const FText LoadedQueueName = FText::Format(
+		LoadedQueueFormatText, !QueueName.IsEmpty() ? FText::FromString(QueueName) : NoLoadedQueueText);
+
+	MenuBuilder.AddMenuEntry(
+		LoadedQueueName,
+		LoadedQueueName,
+		FSlateIcon(),
+		FUIAction(
+			FExecuteAction(),
+			FCanExecuteAction::CreateLambda([]() { return false; })),
+		NAME_None,
+		EUserInterfaceActionType::None
+	);
+
+	MenuBuilder.AddMenuSeparator();
+
+	MenuBuilder.AddMenuEntry(
+		LOCTEXT("SaveQueue_Text", "Save Queue"),
+		LOCTEXT("SaveQueue_Tip", "Save the current configuration in its existing preset which can be shared between multiple jobs, or imported later as the base of a new configuration."),
+		FSlateIcon(FAppStyle::Get().GetStyleSetName(), "AssetEditor.SaveAsset"),
+		FUIAction(
+			FExecuteAction::CreateSP(this, &SMoviePipelineNewQueuePanel::OnSaveAsset),
+			FCanExecuteAction::CreateSP(this, &SMoviePipelineNewQueuePanel::IsQueueDirty)
+		)
+	);
+
+	MenuBuilder.AddMenuEntry(
+		LOCTEXT("SaveAsQueue_Text", "Save Queue As"),
+		LOCTEXT("SaveAsQueue_Tip", "Save the current configuration as a new preset that can be shared between multiple jobs, or imported later as the base of a new configuration."),
+		FSlateIcon(FAppStyle::Get().GetStyleSetName(), "AssetEditor.SaveAssetAs"),
+		FUIAction(FExecuteAction::CreateSP(this, &SMoviePipelineNewQueuePanel::OnSaveAsAsset))
+	);
+
+	FAssetPickerConfig AssetPickerConfig;
+	{
+		AssetPickerConfig.SelectionMode = ESelectionMode::Single;
+		AssetPickerConfig.InitialAssetViewType = EAssetViewType::Column;
+		AssetPickerConfig.bFocusSearchBoxWhenOpened = true;
+		AssetPickerConfig.bAllowNullSelection = false;
+		AssetPickerConfig.bShowBottomToolbar = true;
+		AssetPickerConfig.bAutohideSearchBar = false;
+		AssetPickerConfig.bAllowDragging = false;
+		AssetPickerConfig.bCanShowClasses = false;
+		AssetPickerConfig.bShowPathInColumnView = true;
+		AssetPickerConfig.bShowTypeInColumnView = false;
+		AssetPickerConfig.bSortByPathInColumnView = false;
+
+		AssetPickerConfig.AssetShowWarningText = LOCTEXT("NoQueueAssets_Warning", "No Queues Found");
+		AssetPickerConfig.Filter.ClassPaths.Add(UMoviePipelineQueue::StaticClass()->GetClassPathName());
+		AssetPickerConfig.Filter.bRecursiveClasses = true;
+		AssetPickerConfig.OnAssetSelected = FOnAssetSelected::CreateSP(this, &SMoviePipelineNewQueuePanel::OnImportSavedQueueAsset);
+	}
+
+	MenuBuilder.BeginSection(NAME_None, LOCTEXT("LoadQueue_MenuSection", "Import Queue"));
+	{
+		TSharedRef<SWidget> PresetPicker = SNew(SBox)
+			.MinDesiredWidth(400.f)
+			.MinDesiredHeight(400.f)
+			[
+				ContentBrowser.CreateAssetPicker(AssetPickerConfig)
+			];
+
+		MenuBuilder.AddWidget(PresetPicker, FText(), true, false);
+	}
+	MenuBuilder.EndSection();
+
+	return MenuBuilder.MakeWidget();
+}
+
+FText SMoviePipelineNewQueuePanel::GetQueueMenuButtonText() const
+{
+	FText QueueName = FText::FromString(GetQueueOriginName());
+	if (QueueName.IsEmpty())
+	{
+		QueueName = LOCTEXT("QueueSaveMenuUnsavedConfig_Text", "Unsaved Queue");
+	}
+	return FText::Format(FText::FromString(TEXT("{0}{1}")),
+		QueueName,
+		FText::FromString(IsQueueDirty() ? TEXT(" *") : TEXT("")));
+}
+
+void SMoviePipelineNewQueuePanel::OnSaveAsset()
+{
+	UE::MovieRenderPipelineEditor::Private::SaveQueue();
+}
+
+void SMoviePipelineNewQueuePanel::OnSaveAsAsset()
+{
+	UE::MovieRenderPipelineEditor::Private::SaveQueueAs();
+}
+
+void SMoviePipelineNewQueuePanel::OnImportSavedQueueAsset(const FAssetData& InPresetAsset) const
+{
+	FSlateApplication::Get().DismissAllMenus();
+	UE::MovieRenderPipelineEditor::Private::ImportSavedQueueAsset(InPresetAsset);
+}
+
+bool SMoviePipelineNewQueuePanel::IsQueueDirty() const
+{
+	return UE::MovieRenderPipelineEditor::Private::IsQueueDirty();
+}
+
+FString SMoviePipelineNewQueuePanel::GetQueueOriginName() const
+{
+	return UE::MovieRenderPipelineEditor::Private::GetQueueOriginName();
+}
+
+#undef LOCTEXT_NAMESPACE // SMoviePipelineNewQueuePanel
diff --git a/Engine/Plugins/MovieScene/MovieRenderPipeline/Source/MovieRenderPipelineEditor/Private/Widgets/SMoviePipelineNewQueuePanel.h b/Engine/Plugins/MovieScene/MovieRenderPipeline/Source/MovieRenderPipelineEditor/Private/Widgets/SMoviePipelineNewQueuePanel.h
new file mode 100644
--- /dev/null
+++ b/Engine/Plugins/MovieScene/MovieRenderPipeline/Source/MovieRenderPipelineEditor/Private/Widgets/SMoviePipelineNewQueuePanel.h
@@ -0,0 +1,95 @@
+// Copyright Epic Games, Inc. All Rights Reserved.
+
+#pragma once
+
+#include "Widgets/DeclarativeSyntaxSupport.h"
+#include "Widgets/SCompoundWidget.h"
+
+class UMoviePipelineConfigBase;
+class UMoviePipelineQueue;
+class SMoviePipelineQueueEditor;
+class SWindow;
+class UMoviePipelineExecutorJob;
+class UMoviePipelineExecutorShot;
+class IDetailsView;
+struct FAssetData;
+
+/**
+ * New iteration of the Movie Render Pipeline queue panel widget.
+ *
+ * Provides the UI for managing render jobs: adding/removing sequences, editing
+ * graph or legacy configs, saving/loading queue presets, and dispatching local
+ * or remote renders. Shared business logic lives in MoviePipelineQueuePanelUtils.h.
+ *
+ * Activated by setting the cvar MovieRenderPipeline.ShowNewQueuePanel to 1
+ * and reopening the Movie Render Queue tab.
+ */
+class SMoviePipelineNewQueuePanel : public SCompoundWidget
+{
+public:
+	SLATE_BEGIN_ARGS(SMoviePipelineNewQueuePanel) {}
+	SLATE_END_ARGS()
+
+	void Construct(const FArguments& InArgs);
+
+private:
+	/** Kicks off a local (in-process) render of the queue. */
+	void OnRenderLocalRequested();
+	bool IsRenderLocalEnabled() const;
+
+	/** Kicks off a remote render of the queue via the configured executor. */
+	void OnRenderRemoteRequested();
+	bool IsRenderRemoteEnabled() const;
+
+	/** Returns true if any render option (local or remote) is currently available. */
+	bool IsRenderEnabled() const;
+
+	/** Builds the render mode dropdown menu (local/remote). */
+	TSharedRef<SWidget> OnGenerateRenderMenu();
+
+	/** Opens the config editor (graph or legacy) for the given job/shot. */
+	void OnEditJobConfigRequested(TWeakObjectPtr<UMoviePipelineExecutorJob> InJob, TWeakObjectPtr<UMoviePipelineExecutorShot> InShot);
+
+	/** Called when an existing preset is chosen for a job or shot. */
+	void OnJobPresetChosen(TWeakObjectPtr<UMoviePipelineExecutorJob> InJob, TWeakObjectPtr<UMoviePipelineExecutorShot> InShot);
+
+	/** Applies a modified config back to the owning job/shot. */
+	void OnConfigUpdatedForJob(TWeakObjectPtr<UMoviePipelineExecutorJob> InJob, TWeakObjectPtr<UMoviePipelineExecutorShot> InShot, UMoviePipelineConfigBase* InConfig);
+
+	/** Applies a preset-based config back to the owning job/shot. */
+	void OnConfigUpdatedForJobToPreset(TWeakObjectPtr<UMoviePipelineExecutorJob> InJob, TWeakObjectPtr<UMoviePipelineExecutorShot> InShot, UMoviePipelineConfigBase* InConfig);
+
+	/** Destroys the currently open config editor window. */
+	void OnConfigWindowClosed();
+
+	/** Responds to job/shot selection changes in the queue editor. */
+	void OnSelectionChanged(const TArray<UMoviePipelineExecutorJob*>& InSelectedJobs, const TArray<UMoviePipelineExecutorShot*>& InSelectedShots);
+
+	/** Returns the widget-switcher index: 0 for details view, 1 for "no selection" placeholder. */
+	int32 GetDetailsViewWidgetIndex() const;
+	bool IsDetailsViewEnabled() const;
+
+	/** Builds the save/load queue dropdown menu. */
+	TSharedRef<SWidget> OnGenerateSavedQueuesMenu();
+	FText GetQueueMenuButtonText() const;
+
+	/** Thin wrappers that delegate to utils and handle Slate-specific concerns (e.g. menu dismissal). */
+	void OnSaveAsset();
+	void OnSaveAsAsset();
+	void OnImportSavedQueueAsset(const FAssetData& InPresetAsset) const;
+	bool IsQueueDirty() const;
+	FString GetQueueOriginName() const;
+
+private:
+	/** The main queue editor list widget. */
+	TSharedPtr<SMoviePipelineQueueEditor> PipelineQueueEditorWidget;
+
+	/** Details panel showing properties for the selected job(s) or shot(s). */
+	TSharedPtr<IDetailsView> JobDetailsPanelWidget;
+
+	/** The config-editor child window, if one is currently open. */
+	TWeakPtr<SWindow> WeakEditorWindow;
+
+	/** Cached count of selected jobs, used to drive the details-view widget switcher. */
+	int32 NumSelectedJobs = 0;
+};
diff --git a/Engine/Plugins/MovieScene/MovieRenderPipeline/Source/MovieRenderPipelineEditor/Private/Widgets/SMoviePipelineQueuePanel.cpp b/Engine/Plugins/MovieScene/MovieRenderPipeline/Source/MovieRenderPipelineEditor/Private/Widgets/SMoviePipelineQueuePanel.cpp
--- a/Engine/Plugins/MovieScene/MovieRenderPipeline/Source/MovieRenderPipelineEditor/Private/Widgets/SMoviePipelineQueuePanel.cpp
+++ b/Engine/Plugins/MovieScene/MovieRenderPipeline/Source/MovieRenderPipelineEditor/Private/Widgets/SMoviePipelineQueuePanel.cpp
@@ -13,12 +13,12 @@
 #include "MovieRenderPipelineStyle.h"
 #include "Sections/MovieSceneCinematicShotSection.h"
 #include "MoviePipelineQueue.h"
+#include "MoviePipelineQueuePanelUtils.h"
 #include "MoviePipelinePrimaryConfig.h"
 #include "MoviePipelineQueueSubsystem.h"
 #include "MoviePipelineTelemetry.h"
 #include "Graph/MovieGraphConfig.h"
 #include "Graph/MovieGraphAssetToolkit.h"
-#include "Graph/MovieGraphConfigFactory.h"
 
 // Slate Includes
 #include "Widgets/SBoxPanel.h"
@@ -52,10 +52,6 @@
 // UnrealEd Includes
 #include "ScopedTransaction.h"
 #include "AssetRegistry/AssetData.h"
-#include "AssetRegistry/AssetRegistryModule.h"
-#include "FileHelpers.h"
-#include "AssetToolsModule.h"
-#include "Misc/FileHelper.h"
 
 #define LOCTEXT_NAMESPACE "SMoviePipelineQueuePanel"
 
@@ -423,7 +419,7 @@ void SMoviePipelineQueuePanel::OnEditJobConfigRequested(TWeakObjectPtr<UMoviePip
 		// If the graph preset is not valid, create a new graph
 		if (!GraphToEdit)
 		{
-			GraphToEdit = GenerateNewShotSubgraph(InJob.Get(), InShot.Get());
+			GraphToEdit = UE::MovieRenderPipelineEditor::Private::GenerateNewShotSubgraph(InJob.Get(), InShot.Get());
 
 			// Don't do anything if the user canceled out of the graph creation process, or there was an error creating the graph
 			if (!GraphToEdit)
@@ -672,213 +668,31 @@ FText SMoviePipelineQueuePanel::GetQueueMenuButtonText() const
 		FText::FromString(IsQueueDirty() ? TEXT(" *") : TEXT("")));
 }
 
-bool SMoviePipelineQueuePanel::OpenSaveDialog(const FString& InDefaultPath, const FString& InNewNameSuggestion, FString& OutPackageName)
-{
-	FSaveAssetDialogConfig SaveAssetDialogConfig;
-	{
-		SaveAssetDialogConfig.DefaultPath = InDefaultPath;
-		SaveAssetDialogConfig.DefaultAssetName = InNewNameSuggestion;
-		SaveAssetDialogConfig.AssetClassNames.Add(UMoviePipelineQueue::StaticClass()->GetClassPathName());
-		SaveAssetDialogConfig.ExistingAssetPolicy = ESaveAssetDialogExistingAssetPolicy::AllowButWarn;
-		SaveAssetDialogConfig.DialogTitleOverride = LOCTEXT("SaveQueueAssetDialogTitle", "Save Queue Asset");
-	}
-
-	FContentBrowserModule& ContentBrowserModule = FModuleManager::LoadModuleChecked<FContentBrowserModule>("ContentBrowser");
-	FString SaveObjectPath = ContentBrowserModule.Get().CreateModalSaveAssetDialog(SaveAssetDialogConfig);
-
-	if (!SaveObjectPath.IsEmpty())
-	{
-		OutPackageName = FPackageName::ObjectPathToPackageName(SaveObjectPath);
-		return true;
-	}
-
-	return false;
-}
-
-bool SMoviePipelineQueuePanel::GetSavePresetPackageName(const FString& InExistingName, FString& OutName)
-{
-	UMovieRenderPipelineProjectSettings* ConfigSettings = GetMutableDefault<UMovieRenderPipelineProjectSettings>();
-
-	// determine default package path
-	const FString DefaultSaveDirectory = ConfigSettings->PresetSaveDir.Path;
-
-	FString DialogStartPath;
-	FPackageName::TryConvertFilenameToLongPackageName(DefaultSaveDirectory, DialogStartPath);
-	if (DialogStartPath.IsEmpty())
-	{
-		DialogStartPath = TEXT("/Game");
-	}
-
-	// determine default asset name
-	FString DefaultName = InExistingName;
-
-	FString UniquePackageName;
-	FString UniqueAssetName;
-
-	FAssetToolsModule& AssetToolsModule = FModuleManager::LoadModuleChecked<FAssetToolsModule>("AssetTools");
-	AssetToolsModule.Get().CreateUniqueAssetName(DialogStartPath / DefaultName, TEXT(""), UniquePackageName, UniqueAssetName);
-
-	FString DialogStartName = FPaths::GetCleanFilename(UniqueAssetName);
-
-	FString UserPackageName;
-	FString NewPackageName;
-
-	// get destination for asset
-	bool bFilenameValid = false;
-	while (!bFilenameValid)
-	{
-		if (!OpenSaveDialog(DialogStartPath, DialogStartName, UserPackageName))
-		{
-			return false;
-		}
-
-		NewPackageName = UserPackageName;
-
-		FText OutError;
-		bFilenameValid = FFileHelper::IsFilenameValidForSaving(NewPackageName, OutError);
-	}
-
-	// Update to the last location they saved to so it remembers their settings next time.
-	ConfigSettings->PresetSaveDir.Path = FPackageName::GetLongPackagePath(UserPackageName);
-	ConfigSettings->SaveConfig();
-	OutName = MoveTemp(NewPackageName);
-	return true;
-}
-
 void SMoviePipelineQueuePanel::OnSaveAsset()
 {
-	UMoviePipelineQueue* QueueOrigin = GetQueueOrigin();
-	
-	// If the queue origin is not known, then treat the "Save" as a "Save As"
-	if (!QueueOrigin)
-	{
-		OnSaveAsAsset();
-		return;
-	}
-	
-	SaveTransientQueueToAsset(QueueOrigin);
+	UE::MovieRenderPipelineEditor::Private::SaveQueue();
 }
 
 void SMoviePipelineQueuePanel::OnSaveAsAsset()
 {
-	UMoviePipelineQueueSubsystem* Subsystem = GEditor->GetEditorSubsystem<UMoviePipelineQueueSubsystem>();
-	check(Subsystem);
-	UMoviePipelineQueue* CurrentQueue = Subsystem->GetQueue();
-
-	FString PackageName;
-	if (!GetSavePresetPackageName(CurrentQueue->GetName(), PackageName))
-	{
-		return;
-	}
-	
-	// Saving into a new package
-	const FString NewAssetName = FPackageName::GetLongPackageAssetName(PackageName);
-	UPackage* NewPackage = CreatePackage(*PackageName);
-	NewPackage->MarkAsFullyLoaded();
-	UMoviePipelineQueue* DuplicateQueue = DuplicateObject<UMoviePipelineQueue>(CurrentQueue, NewPackage, *NewAssetName);
-
-	SaveTransientQueueToAsset(DuplicateQueue);
-}
-
-void SMoviePipelineQueuePanel::SaveTransientQueueToAsset(UMoviePipelineQueue* DestinationQueue) const
-{
-	if (!DestinationQueue)
-	{
-		return;
-	}
-
-	UMoviePipelineQueueSubsystem* Subsystem = GEditor->GetEditorSubsystem<UMoviePipelineQueueSubsystem>();
-	check(Subsystem);
-	UMoviePipelineQueue* CurrentQueue = Subsystem->GetQueue();
-
-	DestinationQueue->CopyFrom(CurrentQueue);
-	DestinationQueue->SetQueueOrigin(nullptr);
-	DestinationQueue->MarkPackageDirty();
-	DestinationQueue->SetFlags(RF_Public | RF_Standalone | RF_Transactional);
-
-	FAssetRegistryModule::AssetCreated(DestinationQueue);
-
-	const bool bCheckDirty = false;
-	const bool bPromptToSave = false;
-	FEditorFileUtils::EPromptReturnCode PromptReturnCode = FEditorFileUtils::PromptForCheckoutAndSave({ DestinationQueue->GetPackage() }, bCheckDirty, bPromptToSave);
-	if (PromptReturnCode == FEditorFileUtils::EPromptReturnCode::PR_Success)
-	{
-		constexpr bool bPromptOnReplacingDirtyQueue = false;
-		Subsystem->LoadQueue(DestinationQueue, bPromptOnReplacingDirtyQueue);
-	}
+	UE::MovieRenderPipelineEditor::Private::SaveQueueAs();
 }
 
 void SMoviePipelineQueuePanel::OnImportSavedQueueAsset(const FAssetData& InPresetAsset) const
 {
 	FSlateApplication::Get().DismissAllMenus();
 
-	UMoviePipelineQueueSubsystem* Subsystem = GEditor->GetEditorSubsystem<UMoviePipelineQueueSubsystem>();
-	check(Subsystem);
-
-	UMoviePipelineQueue* SavedQueue = CastChecked<UMoviePipelineQueue>(InPresetAsset.GetAsset());
-	Subsystem->LoadQueue(SavedQueue);
+	UE::MovieRenderPipelineEditor::Private::ImportSavedQueueAsset(InPresetAsset);
 }
 
 bool SMoviePipelineQueuePanel::IsQueueDirty() const
 {
-	const UMoviePipelineQueueSubsystem* Subsystem = GEditor->GetEditorSubsystem<UMoviePipelineQueueSubsystem>();
-	check(Subsystem);
-
-	return Subsystem->IsQueueDirty();
-}
-
-UMoviePipelineQueue* SMoviePipelineQueuePanel::GetQueueOrigin() const
-{
-	UMoviePipelineQueueSubsystem* Subsystem = GEditor->GetEditorSubsystem<UMoviePipelineQueueSubsystem>();
-	check(Subsystem);
-
-	const UMoviePipelineQueue* Queue = Subsystem->GetQueue();
-
-	return Queue ? Queue->GetQueueOrigin() : nullptr;
+	return UE::MovieRenderPipelineEditor::Private::IsQueueDirty();
 }
 
 FString SMoviePipelineQueuePanel::GetQueueOriginName() const
 {
-	const UMoviePipelineQueue* QueueOrigin = GetQueueOrigin();
-	if (!QueueOrigin)
-	{
-		return FString();
-	}
-
-	if (const UPackage* Package = QueueOrigin->GetPackage())
-	{
-		FString OutPackageRoot, OutPackagePath, OutPackageName;
-		if (FPackageName::SplitLongPackageName(Package->GetName(), OutPackageRoot, OutPackagePath, OutPackageName))
-		{
-			return OutPackageName;
-		}
-	}
-
-	return FString();
-}
-
-UMovieGraphConfig* SMoviePipelineQueuePanel::GenerateNewShotSubgraph(const UMoviePipelineExecutorJob* InJob, UMoviePipelineExecutorShot* InShot) const
-{
-	UMovieGraphConfigFactory* GraphFactory = NewObject<UMovieGraphConfigFactory>();
-	GraphFactory->InitialSubgraphAsset = InJob->GetGraphPreset();
-			
-	// Make the new graph via save dialog
-	const FAssetToolsModule& AssetToolsModule = FAssetToolsModule::GetModule();
-	UObject* NewAsset = AssetToolsModule.Get().CreateAssetWithDialog(GraphFactory->GetSupportedClass(), GraphFactory);
-
-	// Don't ensure here because a "cancel" in the dialog can cause the returned asset to be null
-	if (UMovieGraphConfig* NewGraph = Cast<UMovieGraphConfig>(NewAsset))
-	{
-		// Save out the new graph and assign it to the shot
-		constexpr bool bOnlyDirty = false;
-		UEditorLoadingAndSavingUtils::SavePackages({ NewGraph->GetPackage() }, bOnlyDirty);
-		
-		InShot->SetGraphPreset(NewGraph);
-		
-		return NewGraph;
-	}
-
-	return nullptr;
+	return UE::MovieRenderPipelineEditor::Private::GetQueueOriginName();
 }
 
 #undef LOCTEXT_NAMESPACE // SMoviePipelineQueuePanel
\ No newline at end of file
diff --git a/Engine/Plugins/MovieScene/MovieRenderPipeline/Source/MovieRenderPipelineEditor/Private/Widgets/SMoviePipelineQueuePanel.h b/Engine/Plugins/MovieScene/MovieRenderPipeline/Source/MovieRenderPipelineEditor/Private/Widgets/SMoviePipelineQueuePanel.h
--- a/Engine/Plugins/MovieScene/MovieRenderPipeline/Source/MovieRenderPipelineEditor/Private/Widgets/SMoviePipelineQueuePanel.h
+++ b/Engine/Plugins/MovieScene/MovieRenderPipeline/Source/MovieRenderPipelineEditor/Private/Widgets/SMoviePipelineQueuePanel.h
@@ -54,26 +54,12 @@ private:
 
 	TSharedRef<SWidget> OnGenerateSavedQueuesMenu();
 	FText GetQueueMenuButtonText() const;
-	bool OpenSaveDialog(const FString& InDefaultPath, const FString& InNewNameSuggestion, FString& OutPackageName);
-	bool GetSavePresetPackageName(const FString& InExistingName, FString& OutName);
 	void OnSaveAsset();
 	void OnSaveAsAsset();
-	void SaveTransientQueueToAsset(UMoviePipelineQueue* DestinationQueue) const;
 	void OnImportSavedQueueAsset(const FAssetData& InPresetAsset) const;
 	bool IsQueueDirty() const;
-	UMoviePipelineQueue* GetQueueOrigin() const;
 	FString GetQueueOriginName() const;
 
-	/**
-	 * Generates and returns a new shot subgraph based on the given job and shot. Returns nullptr if there was an issue creating the graph. If the
-	 * graph was successfully created, it will be assigned to the shot.
-	 */
-	class UMovieGraphConfig* GenerateNewShotSubgraph(const UMoviePipelineExecutorJob* InJob, UMoviePipelineExecutorShot* InShot) const;
-
-private:
-	/** Allocates a transient preset so that the user can use the pipeline without saving it to an asset first. */
-	//UMoviePipelineConfigBase* AllocateTransientPreset();
-
 
 private:
 	/** The main movie pipeline queue editor widget */
diff --git a/Engine/Plugins/MovieScene/MovieRenderPipeline/Source/MovieRenderPipelineEditor/Private/Widgets/SMoviePipelineQueueTabContent.cpp b/Engine/Plugins/MovieScene/MovieRenderPipeline/Source/MovieRenderPipelineEditor/Private/Widgets/SMoviePipelineQueueTabContent.cpp
--- a/Engine/Plugins/MovieScene/MovieRenderPipeline/Source/MovieRenderPipelineEditor/Private/Widgets/SMoviePipelineQueueTabContent.cpp
+++ b/Engine/Plugins/MovieScene/MovieRenderPipeline/Source/MovieRenderPipelineEditor/Private/Widgets/SMoviePipelineQueueTabContent.cpp
@@ -1,6 +1,7 @@
 // Copyright Epic Games, Inc. All Rights Reserved.
 
 #include "Widgets/SMoviePipelineQueueTabContent.h"
+#include "SMoviePipelineNewQueuePanel.h"
 #include "SMoviePipelineQueuePanel.h"
 
 // Analytics
@@ -9,16 +10,36 @@
 
 #define LOCTEXT_NAMESPACE "SMoviePipelineQueueTabContent"
 
+static TAutoConsoleVariable<int32> CVarMovieRenderPipelineShowNewQueuePanel(
+	TEXT("MovieRenderPipeline.ShowNewQueuePanel"),
+	0,
+	TEXT("When set to 1, the Movie Render Queue tab uses the new queue panel.\n")
+	TEXT("When set to 0 (default), the original queue panel is used.\n")
+	TEXT("Changes take effect the next time the queue tab is opened."),
+	ECVF_Default);
 
 void SMoviePipelineQueueTabContent::Construct(const FArguments& InArgs)
-{    
-	ChildSlot
-		.HAlign(HAlign_Fill)
-		.VAlign(VAlign_Fill)
-		[
-			SAssignNew(WeakPanel, SMoviePipelineQueuePanel)
-		];
+{
+	const bool bUseNewPanel = (CVarMovieRenderPipelineShowNewQueuePanel.GetValueOnGameThread() != 0);
 
+	if (bUseNewPanel)
+	{
+		ChildSlot
+			.HAlign(HAlign_Fill)
+			.VAlign(VAlign_Fill)
+			[
+				SAssignNew(WeakNewPanel, SMoviePipelineNewQueuePanel)
+			];
+	}
+	else
+	{
+		ChildSlot
+			.HAlign(HAlign_Fill)
+			.VAlign(VAlign_Fill)
+			[
+				SAssignNew(WeakPanel, SMoviePipelineQueuePanel)
+			];
+	}
 
 	if (FEngineAnalytics::IsAvailable())
 	{
diff --git a/Engine/Plugins/MovieScene/MovieRenderPipeline/Source/MovieRenderPipelineEditor/Private/Widgets/SMoviePipelineQueueTabContent.h b/Engine/Plugins/MovieScene/MovieRenderPipeline/Source/MovieRenderPipelineEditor/Private/Widgets/SMoviePipelineQueueTabContent.h
--- a/Engine/Plugins/MovieScene/MovieRenderPipeline/Source/MovieRenderPipelineEditor/Private/Widgets/SMoviePipelineQueueTabContent.h
+++ b/Engine/Plugins/MovieScene/MovieRenderPipeline/Source/MovieRenderPipelineEditor/Private/Widgets/SMoviePipelineQueueTabContent.h
@@ -5,6 +5,7 @@
 #include "Widgets/DeclarativeSyntaxSupport.h"
 #include "Widgets/SCompoundWidget.h"
 
+class SMoviePipelineNewQueuePanel;
 class SMoviePipelineQueuePanel;
 
 class SMoviePipelineQueueTabContent : public SCompoundWidget
@@ -18,4 +19,5 @@ public:
 
 private:
 	TWeakPtr<SMoviePipelineQueuePanel> WeakPanel;
+	TWeakPtr<SMoviePipelineNewQueuePanel> WeakNewPanel;
 };
