MoviePipeline: Added a new "Visible in Reflections" property to the Modifier node in the graph.

#Synced-CL 51731107

diff --git a/Engine/Plugins/MovieScene/MovieRenderPipeline/Source/MovieRenderPipelineCore/Private/Graph/MovieGraphRenderLayerSubsystem.cpp b/Engine/Plugins/MovieScene/MovieRenderPipeline/Source/MovieRenderPipelineCore/Private/Graph/MovieGraphRenderLayerSubsystem.cpp
--- a/Engine/Plugins/MovieScene/MovieRenderPipeline/Source/MovieRenderPipelineCore/Private/Graph/MovieGraphRenderLayerSubsystem.cpp
+++ b/Engine/Plugins/MovieScene/MovieRenderPipeline/Source/MovieRenderPipelineCore/Private/Graph/MovieGraphRenderLayerSubsystem.cpp
@@ -190,6 +190,7 @@ UMovieGraphRenderPropertyModifier::UMovieGraphRenderPropertyModifier()
 	, bCastShadowWhileHidden(false)
 	, bAffectIndirectLightingWhileHidden(false)
 	, bHoldout(false)
+	, bVisibleInReflections(true)
 	, bProcessEditorOnlyActors(false)
 {
 	// Note: The default modifier values here reflect the defaults on the scene component. If a modifier property is marked as overridden, the
@@ -237,6 +238,7 @@ void UMovieGraphRenderPropertyModifier::ApplyModifier(const UWorld* World)
 			CachedState.bCastShadowWhileHidden = AsPrimitiveComponent->bCastHiddenShadow;
 			CachedState.bAffectIndirectLightingWhileHidden = AsPrimitiveComponent->bAffectIndirectLightingWhileHidden;
 			CachedState.bHoldout = AsPrimitiveComponent->bHoldout;
+			CachedState.bVisibleInReflections = AsPrimitiveComponent->bVisibleInReflections;
 			CachedState.bIsHidden = AsPrimitiveComponent->bHiddenInGame;
 		}
 		// Volumetrics are special cases as they don't inherit from UPrimitiveComponent, and don't support all of the flags.
@@ -403,6 +405,16 @@ void UMovieGraphRenderPropertyModifier::SetActorAndComponentState(const TArray<F
 				AsPrimitiveComponent->SetHoldout(InNewComponentState.bHoldout);
 			}
 
+			if (bOverride_bVisibleInReflections)
+			{
+				// There's no Set*() method for Visible in Reflections, so MarkRenderStateDirty() needs to be called manually
+				if (AsPrimitiveComponent->bVisibleInReflections != InNewComponentState.bVisibleInReflections)
+				{
+					AsPrimitiveComponent->bVisibleInReflections = InNewComponentState.bVisibleInReflections;
+					AsPrimitiveComponent->MarkRenderStateDirty();
+				}
+			}
+
 			if (bOverride_bIsHidden)
 			{
 				AsPrimitiveComponent->SetHiddenInGame(InNewComponentState.bIsHidden);
@@ -459,6 +471,7 @@ void UMovieGraphRenderPropertyModifier::SetActorAndComponentState(const TArray<F
 		NewComponentState.bCastsShadows = bCastsShadows;
 		NewComponentState.bCastShadowWhileHidden = bCastShadowWhileHidden;
 		NewComponentState.bAffectIndirectLightingWhileHidden = bAffectIndirectLightingWhileHidden;
+		NewComponentState.bVisibleInReflections = bVisibleInReflections;
 
 		for (const FComponentState& CachedState : CachedComponentState)
 		{
diff --git a/Engine/Plugins/MovieScene/MovieRenderPipeline/Source/MovieRenderPipelineCore/Public/Graph/MovieGraphRenderLayerSubsystem.h b/Engine/Plugins/MovieScene/MovieRenderPipeline/Source/MovieRenderPipelineCore/Public/Graph/MovieGraphRenderLayerSubsystem.h
--- a/Engine/Plugins/MovieScene/MovieRenderPipeline/Source/MovieRenderPipelineCore/Public/Graph/MovieGraphRenderLayerSubsystem.h
+++ b/Engine/Plugins/MovieScene/MovieRenderPipeline/Source/MovieRenderPipelineCore/Public/Graph/MovieGraphRenderLayerSubsystem.h
@@ -1082,6 +1082,7 @@ private:
 		uint8 bCastShadowWhileHidden : 1 = false;
 		uint8 bAffectIndirectLightingWhileHidden : 1 = false;
 		uint8 bHoldout : 1 = false;
+		uint8 bVisibleInReflections : 1 = true;
 		uint8 bIsHidden : 1 = false;
 	};
 
@@ -1118,6 +1119,9 @@ public:
 	UPROPERTY(EditAnywhere, BlueprintReadWrite, Category = Overrides, meta = (InlineEditConditionToggle))
 	uint8 bOverride_bHoldout : 1;
 
+	UPROPERTY(EditAnywhere, BlueprintReadWrite, Category = Overrides, meta = (InlineEditConditionToggle))
+	uint8 bOverride_bVisibleInReflections : 1;
+
 	UPROPERTY(EditAnywhere, BlueprintReadWrite, Category = Overrides, meta = (InlineEditConditionToggle))
 	uint8 bOverride_bProcessEditorOnlyActors : 1;
 	
@@ -1147,6 +1151,10 @@ public:
 	UPROPERTY(BlueprintReadWrite, EditAnywhere, Category = "Settings", meta = (EditCondition = "bOverride_bHoldout"))
 	uint8 bHoldout : 1;
 
+	/** If true, this component will be visible in ray traced reflections. Screen space traces will ignore this flag, as they directly sample SceneColor. */
+	UPROPERTY(BlueprintReadWrite, EditAnywhere, Category = "Settings", meta = (EditCondition = "bOverride_bVisibleInReflections"))
+	uint8 bVisibleInReflections : 1;
+
 	/**
 	 * If true, this modifier will process editor-only actors (eg, the billboard icons that represent lights). Defaults to false. Generally, processing
 	 * editor-only actors is a waste of cycles because they will never be shown in a render, and there's no point in messing with them. However there
