MoviePipeline: Fixed a bug that caused depth PPMs in the graph to generate corrupt data when using AMD graphics cards.

#Synced-CL 51731107

diff --git a/Engine/Plugins/MovieScene/MovieRenderPipeline/Source/MovieRenderPipelineRenderPasses/Private/Graph/Renderers/MovieGraphDeferredPass.cpp b/Engine/Plugins/MovieScene/MovieRenderPipeline/Source/MovieRenderPipelineRenderPasses/Private/Graph/Renderers/MovieGraphDeferredPass.cpp
--- a/Engine/Plugins/MovieScene/MovieRenderPipeline/Source/MovieRenderPipelineRenderPasses/Private/Graph/Renderers/MovieGraphDeferredPass.cpp
+++ b/Engine/Plugins/MovieScene/MovieRenderPipeline/Source/MovieRenderPipelineRenderPasses/Private/Graph/Renderers/MovieGraphDeferredPass.cpp
@@ -20,16 +20,52 @@
 #include "Materials/Material.h"
 #include "Materials/MaterialInterface.h"
 #include "OpenColorIODisplayExtension.h"
+#include "PostProcess/PostProcessInputs.h"
 #include "SceneManagement.h"
+#include "SceneRenderTargetParameters.h"
 #include "SceneView.h"
 #include "Tasks/Task.h"
 #include "TextureResource.h"
 
-
 #include "MoviePipelineBlueprintLibrary.h" // For overscan calculations
 #include "MoviePipelineUtils.h" // For Sub-Pixel Jitter
 #include "MovieRenderPipelineDataTypes.h" // For the 1D Weight table for accumulation
 
+// Shader parameter struct used to force depth buffer decompression via RDG resource transition.
+// On AMD GPUs, the depth buffer may remain HTile-compressed when Movie Render Graph disables
+// intermediate post-process passes (DOF, MotionBlur, TAA) that would normally trigger decompression.
+// Adding an empty RDG pass with this struct as SRV forces the necessary decompression barrier.
+BEGIN_SHADER_PARAMETER_STRUCT(FBufferVisDepthDecompressParameters, )
+	SHADER_PARAMETER_RDG_TEXTURE(Texture2D<float>, SceneDepth)
+END_SHADER_PARAMETER_STRUCT()
+
+// Scene view extension that forces depth buffer decompression before post-processing begins.
+// This ensures buffer visualization post-process materials (WorldDepth, etc.) can correctly
+// read SceneDepth on all GPU vendors.
+class FMovieGraphDepthDecompressExtension : public ISceneViewExtension
+{
+public:
+	virtual ~FMovieGraphDepthDecompressExtension() = default;
+
+	virtual void PrePostProcessPass_RenderThread(FRDGBuilder& GraphBuilder, const FSceneView& InView, const FPostProcessingInputs& Inputs) override
+	{
+		if (Inputs.SceneTextures)
+		{
+			if (const FRDGTextureRef SceneDepthTexture = (*Inputs.SceneTextures)->SceneDepthTexture)
+			{
+				FBufferVisDepthDecompressParameters* DecompressParams = GraphBuilder.AllocParameters<FBufferVisDepthDecompressParameters>();
+				DecompressParams->SceneDepth = SceneDepthTexture;
+				GraphBuilder.AddPass(
+					RDG_EVENT_NAME("MRG_DecompressDepthForBufferVisualization"),
+					DecompressParams,
+					ERDGPassFlags::Copy | ERDGPassFlags::NeverCull,
+					[](FRDGAsyncTask, FRHICommandList&) {}
+				);
+			}
+		}
+	}
+};
+
 namespace UE::MovieGraph::Rendering
 {
 	namespace MetadataHelper
@@ -630,6 +666,15 @@ void FMovieGraphDeferredPass::Render(const FMovieGraphTraversalContext& InFrameT
 				// The denoiser is disabled during warm-up frames.
 				NewView->FinalPostProcessSettings.PathTracingEnableDenoiser = bIsRenderingState && ParentNodeThisFrame->GetAllowDenoiser();
 
+				// When buffer visualization PPMs are active, register a view extension that forces depth
+				// buffer decompression before post-processing. This prevents AMD HTile corruption when
+				// PPMs read SceneDepth as SRV after intermediate passes that normally trigger decompression
+				// (DOF, MotionBlur, TAA) have been skipped.
+				if (NumValidMaterials > 0)
+				{
+					ViewFamily->ViewExtensions.Add(MakeShareable(new FMovieGraphDepthDecompressExtension()));
+				}
+
 				// Submit the renderer to be rendered
 				GetRendererModule().BeginRenderingViewFamily(&Canvas, ViewFamily.ToSharedPtr().Get());
 
