MoviePipeline: Another small UI pass on the new MRQ queue panel.

#Synced-CL 51731107

diff --git a/Engine/Plugins/MovieScene/MovieRenderPipeline/Source/MovieRenderPipelineEditor/Private/Widgets/SMoviePipelineNewQueuePanel.cpp b/Engine/Plugins/MovieScene/MovieRenderPipeline/Source/MovieRenderPipelineEditor/Private/Widgets/SMoviePipelineNewQueuePanel.cpp
--- a/Engine/Plugins/MovieScene/MovieRenderPipeline/Source/MovieRenderPipelineEditor/Private/Widgets/SMoviePipelineNewQueuePanel.cpp
+++ b/Engine/Plugins/MovieScene/MovieRenderPipeline/Source/MovieRenderPipelineEditor/Private/Widgets/SMoviePipelineNewQueuePanel.cpp
@@ -116,46 +116,22 @@ void SMoviePipelineNewQueuePanel::Construct(const FArguments& InArgs)
 			[
 				SNew(SHorizontalBox)
 
-				// Add a Level Sequence to the queue
+				// Toolbar (FToolBarBuilder-generated)
 				+ SHorizontalBox::Slot()
-				.Padding(MoviePipeline::ButtonOffset)
-				.VAlign(VAlign_Fill)
-				.AutoWidth()
-				[
-					SNew(SPositiveActionButton)
-					.OnGetMenuContent_Lambda([this]() -> TSharedRef<SWidget>
-					{
-						return PipelineQueueEditorWidget->OnGenerateNewJobFromAssetMenu();
-					})
-					.Icon(FAppStyle::Get().GetBrush("Icons.Plus"))
-					.Text(LOCTEXT("AddNewJob_Text", "Add"))
-				]
-
-				// Render queue (local/remote dropdown)
-				+ SHorizontalBox::Slot()
-				.Padding(MoviePipeline::ButtonOffset)
-				.VAlign(VAlign_Fill)
 				.AutoWidth()
+				.VAlign(VAlign_Center)
 				[
-					SNew(SActionButton)
-					.ActionButtonType(EActionButtonType::Primary)
-					.HasDownArrow(true)
-					.Icon(FAppStyle::Get().GetBrush("LevelEditor.OpenCinematic"))
-					.Text(LOCTEXT("RenderQueue_Text", "Render"))
-					.IsEnabled(this, &SMoviePipelineNewQueuePanel::IsRenderEnabled)
-					.OnGetMenuContent(this, &SMoviePipelineNewQueuePanel::OnGenerateRenderMenu)
+					MakeToolbar()
 				]
 
 				// Spacer
 				+ SHorizontalBox::Slot()
-				.VAlign(VAlign_Fill)
-				.HAlign(HAlign_Fill)
 				.FillWidth(1.f)
 				[
 					SNullWidget::NullWidget
 				]
 
-				// Queue save/load dropdown
+				// Queue save/load dropdown (right-aligned, stays outside FToolBarBuilder)
 				+ SHorizontalBox::Slot()
 				.Padding(MoviePipeline::ButtonOffset)
 				.VAlign(VAlign_Center)
@@ -296,6 +272,73 @@ void SMoviePipelineNewQueuePanel::Construct(const FArguments& InArgs)
 	];
 }
 
+TSharedRef<SWidget> SMoviePipelineNewQueuePanel::MakeToolbar()
+{
+	FSlimHorizontalToolBarBuilder ToolbarBuilder(TSharedPtr<const FUICommandList>(), FMultiBoxCustomization::None);
+	ToolbarBuilder.SetLabelVisibility(EVisibility::Collapsed);
+
+	// Section 1: Save / Find
+	ToolbarBuilder.BeginSection("SaveFind");
+	{
+		ToolbarBuilder.AddToolBarButton(
+			FUIAction(FExecuteAction::CreateSP(this, &SMoviePipelineNewQueuePanel::OnSaveButtonClicked)),
+			NAME_None,
+			LOCTEXT("SaveQueueButton_Text", "Save"),
+			LOCTEXT("SaveQueueButton_Tooltip", "Save the current queue. If the queue has not been saved before, prompts for a save location."),
+			FSlateIcon(FAppStyle::GetAppStyleSetName(), "AssetEditor.SaveAsset"));
+
+		ToolbarBuilder.AddToolBarButton(
+			FUIAction(
+				FExecuteAction::CreateSP(this, &SMoviePipelineNewQueuePanel::OnFindInContentBrowserClicked),
+				FCanExecuteAction::CreateSP(this, &SMoviePipelineNewQueuePanel::IsFindInContentBrowserEnabled)),
+			NAME_None,
+			LOCTEXT("FindInContentBrowser_Text", "Browse"),
+			LOCTEXT("FindInContentBrowser_Tooltip", "Select the job's level sequence in the Content Browser."),
+			FSlateIcon(FAppStyle::GetAppStyleSetName(), "Icons.BrowseContent"));
+
+	}
+	ToolbarBuilder.EndSection();
+
+	// Section 2: Add / Render
+	ToolbarBuilder.BeginSection("AddRender");
+	{
+		const FToolBarStyle& ToolBarStyle = ToolbarBuilder.GetStyleSet()->GetWidgetStyle<FToolBarStyle>(ToolbarBuilder.GetStyleName());
+
+		ToolbarBuilder.AddWidget(
+			SNew(SBox)
+			.VAlign(VAlign_Fill)
+			.Padding(ToolBarStyle.ButtonPadding)
+			[
+				SNew(SPositiveActionButton)
+				.OnGetMenuContent_Lambda([this]() -> TSharedRef<SWidget>
+				{
+					return PipelineQueueEditorWidget->OnGenerateNewJobFromAssetMenu();
+				})
+				.Icon(FAppStyle::Get().GetBrush("Icons.Plus"))
+				.Text(LOCTEXT("AddNewJob_Text", "Add"))
+			]);
+
+		ToolbarBuilder.AddSeparator();
+
+		ToolbarBuilder.AddWidget(
+			SNew(SBox)
+			.VAlign(VAlign_Fill)
+			.Padding(ToolBarStyle.ButtonPadding)
+			[
+				SNew(SActionButton)
+				.ActionButtonType(EActionButtonType::Primary)
+				.HasDownArrow(true)
+				.Icon(FAppStyle::Get().GetBrush("LevelEditor.OpenCinematic"))
+				.Text(LOCTEXT("RenderQueue_Text", "Render"))
+				.IsEnabled(this, &SMoviePipelineNewQueuePanel::IsRenderEnabled)
+				.OnGetMenuContent(this, &SMoviePipelineNewQueuePanel::OnGenerateRenderMenu)
+			]);
+	}
+	ToolbarBuilder.EndSection();
+
+	return ToolbarBuilder.MakeWidget();
+}
+
 void SMoviePipelineNewQueuePanel::OnRenderLocalRequested()
 {
 	UE::MovieRenderPipelineEditor::Private::PerformLocalRender();
@@ -679,4 +722,35 @@ FString SMoviePipelineNewQueuePanel::GetQueueOriginName() const
 	return UE::MovieRenderPipelineEditor::Private::GetQueueOriginName();
 }
 
+void SMoviePipelineNewQueuePanel::OnSaveButtonClicked()
+{
+	// SaveQueue() already does Save As when there is no queue origin
+	UE::MovieRenderPipelineEditor::Private::SaveQueue();
+}
+
+void SMoviePipelineNewQueuePanel::OnFindInContentBrowserClicked()
+{
+	for (UMoviePipelineExecutorJob* Job : PipelineQueueEditorWidget->GetSelectedJobs())
+	{
+		if (UObject* SequenceAsset = Job->Sequence.TryLoad())
+		{
+			FContentBrowserModule& ContentBrowserModule = FModuleManager::LoadModuleChecked<FContentBrowserModule>("ContentBrowser");
+			ContentBrowserModule.Get().SyncBrowserToAssets(TArray<UObject*>({ SequenceAsset }));
+			return;
+		}
+	}
+}
+
+bool SMoviePipelineNewQueuePanel::IsFindInContentBrowserEnabled() const
+{
+	for (const UMoviePipelineExecutorJob* Job : PipelineQueueEditorWidget->GetSelectedJobs())
+	{
+		if (Job->Sequence.IsValid())
+		{
+			return true;
+		}
+	}
+	return false;
+}
+
 #undef LOCTEXT_NAMESPACE // SMoviePipelineNewQueuePanel
diff --git a/Engine/Plugins/MovieScene/MovieRenderPipeline/Source/MovieRenderPipelineEditor/Private/Widgets/SMoviePipelineNewQueuePanel.h b/Engine/Plugins/MovieScene/MovieRenderPipeline/Source/MovieRenderPipelineEditor/Private/Widgets/SMoviePipelineNewQueuePanel.h
--- a/Engine/Plugins/MovieScene/MovieRenderPipeline/Source/MovieRenderPipelineEditor/Private/Widgets/SMoviePipelineNewQueuePanel.h
+++ b/Engine/Plugins/MovieScene/MovieRenderPipeline/Source/MovieRenderPipelineEditor/Private/Widgets/SMoviePipelineNewQueuePanel.h
@@ -73,6 +73,18 @@ private:
 	TSharedRef<SWidget> OnGenerateSavedQueuesMenu();
 	FText GetQueueMenuButtonText() const;
 
+	/** Saves the queue (Save As if unsaved, Save if already saved). */
+	void OnSaveButtonClicked();
+
+	/** Selects the selected job's level sequence in the Content Browser. */
+	void OnFindInContentBrowserClicked();
+
+	/** Returns true when a job with a valid sequence is selected. */
+	bool IsFindInContentBrowserEnabled() const;
+
+	/** Builds the main toolbar using FToolBarBuilder. */
+	TSharedRef<SWidget> MakeToolbar();
+
 	/** Thin wrappers that delegate to utils and handle Slate-specific concerns (e.g. menu dismissal). */
 	void OnSaveAsset();
 	void OnSaveAsAsset();
diff --git a/Engine/Plugins/MovieScene/MovieRenderPipeline/Source/MovieRenderPipelineEditor/Private/Widgets/SMoviePipelineQueueEditor.cpp b/Engine/Plugins/MovieScene/MovieRenderPipeline/Source/MovieRenderPipelineEditor/Private/Widgets/SMoviePipelineQueueEditor.cpp
--- a/Engine/Plugins/MovieScene/MovieRenderPipeline/Source/MovieRenderPipelineEditor/Private/Widgets/SMoviePipelineQueueEditor.cpp
+++ b/Engine/Plugins/MovieScene/MovieRenderPipeline/Source/MovieRenderPipelineEditor/Private/Widgets/SMoviePipelineQueueEditor.cpp
@@ -731,7 +731,6 @@ TSharedRef<SWidget> SQueueJobListRow::GenerateWidgetForColumn(const FName& Colum
 			.WidthOverride(16)
 		  	[
 				SNew(SCheckBox)
-				.Style(FMovieRenderPipelineStyle::Get(), "MovieRenderPipeline.Setting.Switch")
 				.IsFocusable(false)
 				.IsChecked(Item.Get(), &FMoviePipelineQueueJobTreeItem::GetCheckState)
 				.OnCheckStateChanged(Item.Get(), &FMoviePipelineQueueJobTreeItem::SetCheckState)
@@ -746,11 +745,22 @@ TSharedRef<SWidget> SQueueJobListRow::GenerateWidgetForColumn(const FName& Colum
 				SNew(SHorizontalBox)
 				+SHorizontalBox::Slot()
 				.AutoWidth()
-				.Padding(0, 0, 6, 0)
+				.Padding(0, 0, 10, 0)
 				[
 					SNew(SExpanderArrow, SharedThis(this))
 				]
 
+				+SHorizontalBox::Slot()
+				.AutoWidth()
+				.Padding(0, 0, 4, 0)
+				.VAlign(VAlign_Center)
+				[
+					SNew(SImage)
+					.Image(FAppStyle::Get().GetBrush("ClassIcon.CameraComponent"))
+					.DesiredSizeOverride(FVector2D(16.f, 16.f))
+					.ColorAndOpacity(FSlateColor::UseForeground())
+				]
+
 				+SHorizontalBox::Slot()
 				.FillWidth(1.f)
 				[
@@ -1354,7 +1364,18 @@ void SQueueShotListRow::Construct(const FArguments& InArgs, const TSharedRef<STa
 
 TSharedRef<SWidget> SQueueShotListRow::GenerateWidgetForColumn(const FName& ColumnName)
 {
-	if (ColumnName == SQueueJobListRow::NAME_JobName)
+	if (ColumnName == SQueueJobListRow::NAME_Enabled)
+	{
+		return SNew(SBox)
+			.WidthOverride(16)
+			[
+				SNew(SCheckBox)
+				.IsFocusable(false)
+				.IsChecked(Item.Get(), &FMoviePipelineShotItem::GetCheckState)
+				.OnCheckStateChanged(Item.Get(), &FMoviePipelineShotItem::SetCheckState)
+			];
+	}
+	else if (ColumnName == SQueueJobListRow::NAME_JobName)
 	{
 		return SNew(SBox)
 		.Padding(2.0f)
@@ -1362,17 +1383,22 @@ TSharedRef<SWidget> SQueueShotListRow::GenerateWidgetForColumn(const FName& Colu
 			SNew(SHorizontalBox)
 			.IsEnabled(Item.Get(), &FMoviePipelineShotItem::IsEnabled)
 
-			// Toggle Checkbox for deciding to render or not.
 			+ SHorizontalBox::Slot()
 			.AutoWidth()
-			.Padding(4, 4)
+			.Padding(0, 0, 0, 0)
+			[
+				SNew(SExpanderArrow, SharedThis(this))
+			]
+
+			+ SHorizontalBox::Slot()
+			.AutoWidth()
+			.Padding(0, 0, 4, 0)
 			.VAlign(VAlign_Center)
 			[
-				SNew(SCheckBox)
-				.Style(FMovieRenderPipelineStyle::Get(), "MovieRenderPipeline.Setting.Switch")
-				.IsFocusable(false)
-				.IsChecked(Item.Get(), &FMoviePipelineShotItem::GetCheckState)
-				.OnCheckStateChanged(Item.Get(), &FMoviePipelineShotItem::SetCheckState)
+				SNew(SImage)
+				.Image(FAppStyle::Get().GetBrush("Sequencer.Tracks.CinematicShot"))
+				.DesiredSizeOverride(FVector2D(16.f, 16.f))
+				.ColorAndOpacity(FSlateColor::UseForeground())
 			]
 
 			// Shot Name Label
@@ -1517,7 +1543,7 @@ void SMoviePipelineQueueEditor::Construct(const FArguments& InArgs)
 
 			+ SHeaderRow::Column(SQueueJobListRow::NAME_Settings)
 			.FillWidth(0.25f)
-			.DefaultLabel(LOCTEXT("QueueHeaderSettings_Text", "Settings"))
+			.DefaultLabel(LOCTEXT("QueueHeaderSettings_Text", "Configuration"))
 
 			+ SHeaderRow::Column(SQueueJobListRow::NAME_Output)
 			.FillWidth(0.45f)
@@ -2027,6 +2053,19 @@ void SMoviePipelineQueueEditor::OnJobSelectionChanged_Impl(TSharedPtr<IMoviePipe
 	OnJobSelectionChanged.ExecuteIfBound(SelectedJobs, SelectedShots);
 }
 
+TArray<UMoviePipelineExecutorJob*> SMoviePipelineQueueEditor::GetSelectedJobs() const
+{
+	TArray<UMoviePipelineExecutorJob*> Jobs;
+	for (const TSharedPtr<IMoviePipelineQueueTreeItem>& Item : TreeView->GetSelectedItems())
+	{
+		if (UMoviePipelineExecutorJob* Job = Item->GetOwningJob())
+		{
+			Jobs.AddUnique(Job);
+		}
+	}
+	return Jobs;
+}
+
 void SMoviePipelineQueueEditor::OnQueueLoaded()
 {
 	const UMoviePipelineQueueSubsystem* Subsystem = GEditor->GetEditorSubsystem<UMoviePipelineQueueSubsystem>();
diff --git a/Engine/Plugins/MovieScene/MovieRenderPipeline/Source/MovieRenderPipelineEditor/Private/Widgets/SMoviePipelineQueueEditor.h b/Engine/Plugins/MovieScene/MovieRenderPipeline/Source/MovieRenderPipelineEditor/Private/Widgets/SMoviePipelineQueueEditor.h
--- a/Engine/Plugins/MovieScene/MovieRenderPipeline/Source/MovieRenderPipelineEditor/Private/Widgets/SMoviePipelineQueueEditor.h
+++ b/Engine/Plugins/MovieScene/MovieRenderPipeline/Source/MovieRenderPipelineEditor/Private/Widgets/SMoviePipelineQueueEditor.h
@@ -51,6 +51,9 @@ public:
 	TArray<TSharedPtr<IMoviePipelineQueueTreeItem>> GetSelectedItems() const { return TreeView->GetSelectedItems(); }
 	void SetSelectedJobs(const TArray<UMoviePipelineExecutorJob*>& InJobs) { PendingJobsToSelect = InJobs; }
 
+	/** Returns the currently selected jobs. */
+	TArray<UMoviePipelineExecutorJob*> GetSelectedJobs() const;
+
 	/** Get the default graph preset from the project settings and assign it to the given job */
 	static void AssignDefaultGraphPresetToJob(UMoviePipelineExecutorJob* InJob);
 
