MoviePipeline: Fixed a bug in the graph that could crash the editor when duplicating a comment node.

#Synced-CL 51731107

diff --git a/Engine/Plugins/MovieScene/MovieRenderPipeline/Source/MovieRenderPipelineEditor/Private/Widgets/Graph/SMovieGraphConfigPanel.cpp b/Engine/Plugins/MovieScene/MovieRenderPipeline/Source/MovieRenderPipelineEditor/Private/Widgets/Graph/SMovieGraphConfigPanel.cpp
--- a/Engine/Plugins/MovieScene/MovieRenderPipeline/Source/MovieRenderPipelineEditor/Private/Widgets/Graph/SMovieGraphConfigPanel.cpp
+++ b/Engine/Plugins/MovieScene/MovieRenderPipeline/Source/MovieRenderPipelineEditor/Private/Widgets/Graph/SMovieGraphConfigPanel.cpp
@@ -388,7 +388,7 @@ bool SMoviePipelineGraphPanel::CanCopySelectedNodes() const
 	{
 		for (UObject* Object : GraphEditorWidget->GetSelectedNodes())
 		{
-			if (const UMoviePipelineEdGraphNodeBase* GraphNode = CastChecked<UMoviePipelineEdGraphNodeBase>(Object))
+			if (const UMoviePipelineEdGraphNodeBase* GraphNode = Cast<UMoviePipelineEdGraphNodeBase>(Object))
 			{
 				// Invalid runtime nodes (likely from an unloaded plugin) cannot be copied.
 				if (!GraphNode->GetRuntimeNode())
@@ -402,6 +402,15 @@ bool SMoviePipelineGraphPanel::CanCopySelectedNodes() const
 					break;
 				}
 			}
+			else if (const UEdGraphNode* EdGraphNode = Cast<UEdGraphNode>(Object))
+			{
+				// Covers non-MRG nodes, but is likely a Comment node.
+				if (EdGraphNode->CanDuplicateNode())
+				{
+					bAtLeastOneCopyableNode = true;
+					break;
+				}
+			}
 		}
 	}
 
