MoviePipeline: Updated the graph to include the ability to edit comment node properties (color, font size, etc).

#Synced-CL 51731107

diff --git a/Engine/Plugins/MovieScene/MovieRenderPipeline/Source/MovieRenderPipelineEditor/Private/Graph/MovieGraphSchema.cpp b/Engine/Plugins/MovieScene/MovieRenderPipeline/Source/MovieRenderPipelineEditor/Private/Graph/MovieGraphSchema.cpp
--- a/Engine/Plugins/MovieScene/MovieRenderPipeline/Source/MovieRenderPipelineEditor/Private/Graph/MovieGraphSchema.cpp
+++ b/Engine/Plugins/MovieScene/MovieRenderPipeline/Source/MovieRenderPipelineEditor/Private/Graph/MovieGraphSchema.cpp
@@ -291,6 +291,21 @@ void UMovieGraphSchema::OnPinConnectionDoubleCicked(UEdGraphPin* PinA, UEdGraphP
 	}
 }
 
+bool UMovieGraphSchema::IsCacheVisualizationOutOfDate(int32 InVisualizationCacheID) const
+{
+	return CurrentCacheRefreshID != InVisualizationCacheID;
+}
+
+int32 UMovieGraphSchema::GetCurrentVisualizationCacheID() const
+{
+	return CurrentCacheRefreshID;
+}
+
+void UMovieGraphSchema::ForceVisualizationCacheClear() const
+{
+	++CurrentCacheRefreshID;
+}
+
 void UMovieGraphSchema::InitMoviePipelineNodeClasses()
 {
 	if (MoviePipelineNodeClasses.Num() > 0)
diff --git a/Engine/Plugins/MovieScene/MovieRenderPipeline/Source/MovieRenderPipelineEditor/Private/Graph/MovieGraphSchema.h b/Engine/Plugins/MovieScene/MovieRenderPipeline/Source/MovieRenderPipelineEditor/Private/Graph/MovieGraphSchema.h
--- a/Engine/Plugins/MovieScene/MovieRenderPipeline/Source/MovieRenderPipelineEditor/Private/Graph/MovieGraphSchema.h
+++ b/Engine/Plugins/MovieScene/MovieRenderPipeline/Source/MovieRenderPipelineEditor/Private/Graph/MovieGraphSchema.h
@@ -32,6 +32,9 @@ public:
 	virtual void SetPinBeingDroppedOnNode(UEdGraphPin* InSourcePin) const override;
 	virtual bool SafeDeleteNodeFromGraph(UEdGraph* Graph, UEdGraphNode* Node) const override;
 	virtual void OnPinConnectionDoubleCicked(UEdGraphPin* PinA, UEdGraphPin* PinB, const FVector2f& GraphPosition) const override;
+	virtual bool IsCacheVisualizationOutOfDate(int32 InVisualizationCacheID) const override;
+	virtual int32 GetCurrentVisualizationCacheID() const override;
+	virtual void ForceVisualizationCacheClear() const override;
 	//~ End EdGraphSchema Interface
 	
 	static void InitMoviePipelineNodeClasses();
@@ -84,6 +87,9 @@ public:
 private:
 	static TArray<UClass*> MoviePipelineNodeClasses;
 
+	/** Counter incremented when cached visualization data (tooltips, titles, etc.) should be refreshed. */
+	mutable int32 CurrentCacheRefreshID = 0;
+
 	/** The pin that is being dropped on another node during a drag-drop operation. Controlled by SupportsDropPinOnNode(). */
 	mutable UEdGraphPin* PinBeingDropped = nullptr;
 };
diff --git a/Engine/Plugins/MovieScene/MovieRenderPipeline/Source/MovieRenderPipelineEditor/Private/Widgets/Graph/SMovieGraphConfigPanel.cpp b/Engine/Plugins/MovieScene/MovieRenderPipeline/Source/MovieRenderPipelineEditor/Private/Widgets/Graph/SMovieGraphConfigPanel.cpp
--- a/Engine/Plugins/MovieScene/MovieRenderPipeline/Source/MovieRenderPipelineEditor/Private/Widgets/Graph/SMovieGraphConfigPanel.cpp
+++ b/Engine/Plugins/MovieScene/MovieRenderPipeline/Source/MovieRenderPipelineEditor/Private/Widgets/Graph/SMovieGraphConfigPanel.cpp
@@ -51,6 +51,7 @@
 #include "FileHelpers.h"
 #include "AssetToolsModule.h"
 #include "Misc/FileHelper.h"
+#include "EdGraphNode_Comment.h"
 #include "EdGraphUtilities.h"
 #include "HAL/PlatformApplicationMisc.h"
 #include "Subsystems/AssetEditorSubsystem.h"
@@ -729,6 +730,12 @@ void SMoviePipelineGraphPanel::OnSelectedNodesChanged(const TSet<UObject*>& NewS
 				}
 			}
 		}
+		else if (UEdGraphNode_Comment* CommentNode = Cast<UEdGraphNode_Comment>(Obj))
+		{
+			// Comment nodes are not MRQ-specific nodes, so they don't have a runtime node.
+			// Add the comment node directly so its properties are displayed in the details panel.
+			SelectedObjects.Add(CommentNode);
+		}
 	}
 
 	OnGraphSelectionChangedEvent.ExecuteIfBound(SelectedObjects);
