MoviePipeline: Updated condition groups in the graph to be expanded by default.

This has the side-effect that condition group expansion state will not be persisted, but for the vast majority of collections this is a non-issue.

#Synced-CL 51731107

diff --git a/Engine/Plugins/MovieScene/MovieRenderPipeline/Source/MovieRenderPipelineEditor/Private/Customizations/Graph/MovieGraphCollectionsCustomization.cpp b/Engine/Plugins/MovieScene/MovieRenderPipeline/Source/MovieRenderPipelineEditor/Private/Customizations/Graph/MovieGraphCollectionsCustomization.cpp
--- a/Engine/Plugins/MovieScene/MovieRenderPipeline/Source/MovieRenderPipelineEditor/Private/Customizations/Graph/MovieGraphCollectionsCustomization.cpp
+++ b/Engine/Plugins/MovieScene/MovieRenderPipeline/Source/MovieRenderPipelineEditor/Private/Customizations/Graph/MovieGraphCollectionsCustomization.cpp
@@ -463,15 +463,8 @@ FMovieGraphConditionGroupBuilder::FMovieGraphConditionGroupBuilder(IDetailLayout
 
 FName FMovieGraphConditionGroupBuilder::GetName() const
 {
-	if (WeakCollection.IsValid())
-	{
-		const TArray<UMovieGraphConditionGroup*>& ConditionGroups = WeakCollection->GetConditionGroups();
-		if (ConditionGroups.IsValidIndex(ConditionGroupIndex))
-		{
-			return FName(ConditionGroups[ConditionGroupIndex]->GetId().ToString());
-		}
-	}
-
+	// NAME_None prevents the details panel from restoring (absent) saved expansion state, which would collapse new nodes.
+	// This ensures InitiallyCollapsed() == false is respected and condition groups start expanded.
 	return NAME_None;
 }
 
@@ -489,7 +482,6 @@ void FMovieGraphConditionGroupBuilder::GenerateHeaderRowContent(FDetailWidgetRow
 	.DragDropHandler(MakeShared<FMovieGraphCollectionDragDropHandler>(bIsConditionGroup, DetailLayoutBuilder, WeakCollection.Get(), ConditionGroupIndex, InvalidQuery, InvalidQueryIndex))
 	.OverrideResetToDefault(FResetToDefaultOverride::Hide())
 	.FilterString(FText::FromString(GroupName))
-	.RowTag(GetName())
 	.PropertyHandleList({ConditionGroupProperty})
 	.NameContent()
 	[
