MoviePipeline: Updated the graph's collection UI to copy the layer operation to the new query when switching the query type.

#Synced-CL 51731107

diff --git a/Engine/Plugins/MovieScene/MovieRenderPipeline/Source/MovieRenderPipelineEditor/Private/Customizations/Graph/MovieGraphCollectionsCustomization.cpp b/Engine/Plugins/MovieScene/MovieRenderPipeline/Source/MovieRenderPipelineEditor/Private/Customizations/Graph/MovieGraphCollectionsCustomization.cpp
--- a/Engine/Plugins/MovieScene/MovieRenderPipeline/Source/MovieRenderPipelineEditor/Private/Customizations/Graph/MovieGraphCollectionsCustomization.cpp
+++ b/Engine/Plugins/MovieScene/MovieRenderPipeline/Source/MovieRenderPipelineEditor/Private/Customizations/Graph/MovieGraphCollectionsCustomization.cpp
@@ -958,9 +958,13 @@ void SMovieGraphCollectionTreeQueryTypeSelectorWidget::SetQueryType(UClass* InNe
 	
 	const int32 ExistingQueryIndex = WeakConditionGroup->GetQueries().Find(WeakQuery.Get());
 
+	// Preserve the operation type from the old query so it doesn't reset to Add.
+	const EMovieGraphConditionGroupQueryOpType PreviousOpType = WeakQuery->GetOperationType();
+
 	// Swap the old query with a new one.
 	// Add the new query at the index of the current query. If the query couldn't be found for some reason, default to the end of the array (-1)
 	UMovieGraphConditionGroupQueryBase* NewQuery = WeakConditionGroup->AddQuery(InNewQueryType, ExistingQueryIndex != INDEX_NONE ? ExistingQueryIndex : -1);
+	NewQuery->SetOperationType(PreviousOpType);
 	WeakConditionGroup->RemoveQuery(WeakQuery.Get());
 	WeakQuery = MakeWeakObjectPtr(NewQuery);
 
