MoviePipeline: Fixed an issue in MRQ that caused incorrect motion blur when temporal sampling is turned on, a high-FPS level sequence is being used, and motion blur on the camera is enabled and set to 0.

#Synced-CL 51731107

diff --git a/Engine/Plugins/MovieScene/MovieRenderPipeline/Source/MovieRenderPipelineCore/Private/MoviePipeline.cpp b/Engine/Plugins/MovieScene/MovieRenderPipeline/Source/MovieRenderPipelineCore/Private/MoviePipeline.cpp
--- a/Engine/Plugins/MovieScene/MovieRenderPipeline/Source/MovieRenderPipelineCore/Private/MoviePipeline.cpp
+++ b/Engine/Plugins/MovieScene/MovieRenderPipeline/Source/MovieRenderPipelineCore/Private/MoviePipeline.cpp
@@ -1484,7 +1484,17 @@ MoviePipeline::FFrameConstantMetrics UMoviePipeline::CalculateShotFrameMetrics(c
 		// If the shutter angle is effectively zero, lie about how long a frame is to prevent divide by zero
 		if (Output.ShutterAnglePercentage < 1.0 / 360.0)
 		{
-			Output.TicksWhileShutterOpen = Output.TicksPerOutputFrame * (1.0 / 360.0);
+			// The minimum shutter fraction must guarantee at least 1 whole tick per temporal
+			// sample, otherwise the per-sample delta floors to 0 ticks and causes zero-delta-time
+			// errors. TemporalSampleCount / TicksPerFrame is the smallest fraction that achieves this.
+			//
+			// Example: 200 ticks/frame, 8 temporal samples
+			//   clamp(8/200, 1/360, 1) = 0.04 -> 200 * 0.04 = 8 ticks / 8 samples = 1 tick/sample
+			//   Total timing: 192 closed + (8 * 1) open = 200 ticks/frame
+			const double MinShutterFraction = FMath::Clamp(
+				(double)AntiAliasingSettings->TemporalSampleCount / Output.TicksPerOutputFrame.AsDecimal(),
+				1.0 / 360.0, 1.0);
+			Output.TicksWhileShutterOpen = Output.TicksPerOutputFrame * MinShutterFraction;
 		}
 		else
 		{
