MoviePipeline: Updated MRQ and MRG to set bPrevAllowBackgroundAudio=true to ensure that audio records correctly even if the editor is in the background.

#Synced-CL 51731107

diff --git a/Engine/Plugins/MovieScene/MovieRenderPipeline/Source/MovieRenderPipelineCore/Private/Graph/MovieGraphDefaultAudioRenderer.cpp b/Engine/Plugins/MovieScene/MovieRenderPipeline/Source/MovieRenderPipelineCore/Private/Graph/MovieGraphDefaultAudioRenderer.cpp
--- a/Engine/Plugins/MovieScene/MovieRenderPipeline/Source/MovieRenderPipelineCore/Private/Graph/MovieGraphDefaultAudioRenderer.cpp
+++ b/Engine/Plugins/MovieScene/MovieRenderPipeline/Source/MovieRenderPipelineCore/Private/Graph/MovieGraphDefaultAudioRenderer.cpp
@@ -16,6 +16,10 @@
 #include "MovieRenderPipelineDataTypes.h"
 #include "Graph/MovieGraphPipeline.h"
 
+#if WITH_EDITOR
+#include "Settings/LevelEditorMiscSettings.h"
+#endif
+
 #include UE_INLINE_GENERATED_CPP_BY_NAME(MovieGraphDefaultAudioRenderer)
 
 void UMovieGraphDefaultAudioRenderer::StartAudioRecording()
@@ -144,6 +148,15 @@ void UMovieGraphDefaultAudioRenderer::SetupAudioRendering()
 	AudioState.PrevUnfocusedAudioMultiplier = FApp::GetUnfocusedVolumeMultiplier();
 	FApp::SetUnfocusedVolumeMultiplier(1.f);
 
+#if WITH_EDITOR
+	// Force background audio on so audio is captured even when the editor is in the background
+	{
+		ULevelEditorMiscSettings* MiscSettings = GetMutableDefault<ULevelEditorMiscSettings>();
+		AudioState.bPrevAllowBackgroundAudio = MiscSettings->bAllowBackgroundAudio;
+		MiscSettings->bAllowBackgroundAudio = true;
+	}
+#endif
+
 	// Will be null if the NRT module wasn't loaded
 	if (IConsoleVariable* AudioRenderEveryTickCvar = IConsoleManager::Get().FindConsoleVariable(TEXT("au.nrt.RenderEveryTick")))
 	{
@@ -167,6 +180,11 @@ void UMovieGraphDefaultAudioRenderer::TeardownAudioRendering() const
 	// Restore previous unfocused audio multiplier, to no longer force audio when unfocused
 	FApp::SetUnfocusedVolumeMultiplier(AudioState.PrevUnfocusedAudioMultiplier);
 
+#if WITH_EDITOR
+	// Restore the original bAllowBackgroundAudio setting
+	GetMutableDefault<ULevelEditorMiscSettings>()->bAllowBackgroundAudio = AudioState.bPrevAllowBackgroundAudio;
+#endif
+
 	// Restore our cached CVar values
 	
 	// This will be null if the NRT wasn't used (module not loaded)
diff --git a/Engine/Plugins/MovieScene/MovieRenderPipeline/Source/MovieRenderPipelineCore/Private/MoviePipelineAudioRendering.cpp b/Engine/Plugins/MovieScene/MovieRenderPipeline/Source/MovieRenderPipelineCore/Private/MoviePipelineAudioRendering.cpp
--- a/Engine/Plugins/MovieScene/MovieRenderPipeline/Source/MovieRenderPipelineCore/Private/MoviePipelineAudioRendering.cpp
+++ b/Engine/Plugins/MovieScene/MovieRenderPipeline/Source/MovieRenderPipelineCore/Private/MoviePipelineAudioRendering.cpp
@@ -16,9 +16,12 @@
 #include "DSP/BufferVectorOperations.h"
 #include "MoviePipelineAntiAliasingSetting.h"
 #include "AudioMixerPlatformNonRealtime.h"
-
 #include "MoviePipelineQueue.h"
 
+#if WITH_EDITOR
+#include "Settings/LevelEditorMiscSettings.h"
+#endif
+
 static FAudioDevice* GetAudioDeviceFromWorldContext(const UObject* WorldContextObject)
 {
 	UWorld* ThisWorld = GEngine->GetWorldFromContextObject(WorldContextObject, EGetWorldErrorMode::LogAndReturnNull);
@@ -45,6 +48,15 @@ void UMoviePipeline::SetupAudioRendering()
 	AudioState.PrevUnfocusedAudioMultiplier = FApp::GetUnfocusedVolumeMultiplier();
 	FApp::SetUnfocusedVolumeMultiplier(1.f);
 
+#if WITH_EDITOR
+	// Force background audio on so audio is captured even when the editor is in the background
+	{
+		ULevelEditorMiscSettings* MiscSettings = GetMutableDefault<ULevelEditorMiscSettings>();
+		AudioState.bPrevAllowBackgroundAudio = MiscSettings->bAllowBackgroundAudio;
+		MiscSettings->bAllowBackgroundAudio = true;
+	}
+#endif
+
 	IConsoleVariable* AudioRenderEveryTickCvar = IConsoleManager::Get().FindConsoleVariable(TEXT("au.nrt.RenderEveryTick"));
 	// Will be null if the NRT module wasn't loaded
 	if (AudioRenderEveryTickCvar)
@@ -60,6 +72,11 @@ void UMoviePipeline::TeardownAudioRendering()
 	// Restore previous unfocused audio multiplier, to no longer force audio when unfocused
 	FApp::SetUnfocusedVolumeMultiplier(AudioState.PrevUnfocusedAudioMultiplier);
 
+#if WITH_EDITOR
+	// Restore the original bAllowBackgroundAudio setting
+	GetMutableDefault<ULevelEditorMiscSettings>()->bAllowBackgroundAudio = AudioState.bPrevAllowBackgroundAudio;
+#endif
+
 	// Restore our cached CVar value.
 	IConsoleVariable* AudioRenderEveryTickCvar = IConsoleManager::Get().FindConsoleVariable(TEXT("au.nrt.RenderEveryTick"));
 
diff --git a/Engine/Plugins/MovieScene/MovieRenderPipeline/Source/MovieRenderPipelineCore/Public/MovieRenderPipelineDataTypes.h b/Engine/Plugins/MovieScene/MovieRenderPipeline/Source/MovieRenderPipelineCore/Public/MovieRenderPipelineDataTypes.h
--- a/Engine/Plugins/MovieScene/MovieRenderPipeline/Source/MovieRenderPipelineCore/Public/MovieRenderPipelineDataTypes.h
+++ b/Engine/Plugins/MovieScene/MovieRenderPipeline/Source/MovieRenderPipelineCore/Public/MovieRenderPipelineDataTypes.h
@@ -1331,6 +1331,7 @@ namespace MoviePipeline
 			, PrevUnfocusedAudioMultiplier(1.f)
 			, PrevRenderEveryTickValue(1)
 			, PrevNeverMuteNRTAudioValue(0)
+			, bPrevAllowBackgroundAudio(false)
 		{}
 
 		bool bIsRecordingAudio;
@@ -1339,6 +1340,9 @@ namespace MoviePipeline
 		int32 PrevRenderEveryTickValue;
 		int32 PrevNeverMuteNRTAudioValue;
 
+		/** The previous value of ULevelEditorMiscSettings::bAllowBackgroundAudio. */
+		bool bPrevAllowBackgroundAudio;
+
 		/** Float Buffers for Movie Segments we've finished rendering. Stored until shutdown. Fully available during Finalize stage. */
 		TArray<FAudioSegment> FinishedSegments;
 
