MoviePipeline: Updated MRG to set the promoted value of CustomPlaybackRangeStart/End to "Sequence Default" if these properties are *not* overridden when they're promoted to a variable. This ensures that default sequence frameranges are rendered, rather than a range containing zero frames.

This changes also introduces a new extensibility point on nodes, GetCustomPromotionDefaultValue(), which allows nodes to customize the default values that variables get when pins are promoted to variables.

#Synced-CL 51731107

diff --git a/Engine/Plugins/MovieScene/MovieRenderPipeline/Source/MovieRenderPipelineCore/Private/Graph/MovieGraphNode.cpp b/Engine/Plugins/MovieScene/MovieRenderPipeline/Source/MovieRenderPipelineCore/Private/Graph/MovieGraphNode.cpp
--- a/Engine/Plugins/MovieScene/MovieRenderPipeline/Source/MovieRenderPipelineCore/Private/Graph/MovieGraphNode.cpp
+++ b/Engine/Plugins/MovieScene/MovieRenderPipeline/Source/MovieRenderPipelineCore/Private/Graph/MovieGraphNode.cpp
@@ -204,6 +204,11 @@ void UMovieGraphNode::UnpromotePropertyFromPin(const FName& PropertyName)
 	UpdatePins();
 }
 
+bool UMovieGraphNode::GetCustomPromotionDefaultValue(const FMovieGraphPropertyInfo& InPropertyInfo, FString& OutSerializedValue) const
+{
+	return false;
+}
+
 TArray<const FProperty*> UMovieGraphNode::GetAllOverrideableProperties() const
 {
 	TArray<const FProperty*> AllProperties;
diff --git a/Engine/Plugins/MovieScene/MovieRenderPipeline/Source/MovieRenderPipelineCore/Private/Graph/Nodes/MovieGraphGlobalOutputSettingNode.cpp b/Engine/Plugins/MovieScene/MovieRenderPipeline/Source/MovieRenderPipelineCore/Private/Graph/Nodes/MovieGraphGlobalOutputSettingNode.cpp
--- a/Engine/Plugins/MovieScene/MovieRenderPipeline/Source/MovieRenderPipelineCore/Private/Graph/Nodes/MovieGraphGlobalOutputSettingNode.cpp
+++ b/Engine/Plugins/MovieScene/MovieRenderPipeline/Source/MovieRenderPipelineCore/Private/Graph/Nodes/MovieGraphGlobalOutputSettingNode.cpp
@@ -3,6 +3,7 @@
 #include "Graph/Nodes/MovieGraphGlobalOutputSettingNode.h"
 
 #include "Algo/Find.h"
+#include "Graph/MovieGraphConfig.h"
 #include "Graph/MovieGraphProjectSettings.h"
 #include "Graph/MovieGraphBlueprintLibrary.h"
 #include "MoviePipelineTelemetry.h"
@@ -91,6 +92,46 @@ EMovieGraphBranchRestriction UMovieGraphGlobalOutputSettingNode::GetBranchRestri
 }
 #endif // WITH_EDITOR
 
+bool UMovieGraphGlobalOutputSettingNode::GetCustomPromotionDefaultValue(const FMovieGraphPropertyInfo& InPropertyInfo, FString& OutSerializedValue) const
+{
+	// For the CustomPlaybackRangeStart/End properties, if the property is NOT overridden,
+	// the promoted variable should default to SequenceDefault (rather than Custom) so that
+	// the variable reflects "use the sequence's own range" semantics.
+	static const FName PlaybackRangeStartName = GET_MEMBER_NAME_CHECKED(UMovieGraphGlobalOutputSettingNode, CustomPlaybackRangeStart);
+	static const FName PlaybackRangeEndName = GET_MEMBER_NAME_CHECKED(UMovieGraphGlobalOutputSettingNode, CustomPlaybackRangeEnd);
+
+	if (InPropertyInfo.Name != PlaybackRangeStartName && InPropertyInfo.Name != PlaybackRangeEndName)
+	{
+		return false;
+	}
+
+	const FProperty* RealProperty = FindFProperty<FProperty>(GetClass(), InPropertyInfo.Name);
+	if (!RealProperty)
+	{
+		return false;
+	}
+
+	const FBoolProperty* OverrideProp = UMovieGraphConfig::FindOverridePropertyForRealProperty(GetClass(), RealProperty);
+	if (!OverrideProp || OverrideProp->GetPropertyValue_InContainer(this))
+	{
+		return false;
+	}
+
+	const FStructProperty* StructProp = CastField<FStructProperty>(RealProperty);
+	if (!StructProp)
+	{
+		return false;
+	}
+
+	FMovieGraphSequencePlaybackRangeBound DesiredValue;
+	DesiredValue.Type = EMovieGraphSequenceRangeType::SequenceDefault;
+
+	// Pass the desired value as both Value and Defaults so ExportText exports all properties
+	// (no diff-against-default skipping). This ensures the full struct is serialized.
+	StructProp->Struct->ExportText(OutSerializedValue, &DesiredValue, &DesiredValue, nullptr, PPF_None, nullptr);
+	return true;
+}
+
 void UMovieGraphGlobalOutputSettingNode::PostLoad()
 {
 	Super::PostLoad();
diff --git a/Engine/Plugins/MovieScene/MovieRenderPipeline/Source/MovieRenderPipelineCore/Public/Graph/MovieGraphNode.h b/Engine/Plugins/MovieScene/MovieRenderPipeline/Source/MovieRenderPipelineCore/Public/Graph/MovieGraphNode.h
--- a/Engine/Plugins/MovieScene/MovieRenderPipeline/Source/MovieRenderPipelineCore/Public/Graph/MovieGraphNode.h
+++ b/Engine/Plugins/MovieScene/MovieRenderPipeline/Source/MovieRenderPipelineCore/Public/Graph/MovieGraphNode.h
@@ -269,6 +269,14 @@ public:
 	UFUNCTION(BlueprintCallable, Category = "Movie Graph")
 	MOVIERENDERPIPELINECORE_API virtual void UnpromotePropertyFromPin(const FName& PropertyName);
 	
+	/**
+	 * Allows subclasses to provide a custom serialized default value when a property is promoted
+	 * to a variable. Called during "Promote to Variable" before writing the variable's default.
+	 * Return true and set OutSerializedValue to override; return false for normal behavior.
+	 */
+	UFUNCTION(BlueprintCallable, Category = "Movie Graph")
+	MOVIERENDERPIPELINECORE_API virtual bool GetCustomPromotionDefaultValue(const FMovieGraphPropertyInfo& InPropertyInfo, FString& OutSerializedValue) const;
+	
 	/**
 	 * Determines if this node type can be added to the graph interactively by a user or via the API when constructing a graph.
 	 * @return true if the object can be added via the API, false otherwise
diff --git a/Engine/Plugins/MovieScene/MovieRenderPipeline/Source/MovieRenderPipelineCore/Public/Graph/Nodes/MovieGraphGlobalOutputSettingNode.h b/Engine/Plugins/MovieScene/MovieRenderPipeline/Source/MovieRenderPipelineCore/Public/Graph/Nodes/MovieGraphGlobalOutputSettingNode.h
--- a/Engine/Plugins/MovieScene/MovieRenderPipeline/Source/MovieRenderPipelineCore/Public/Graph/Nodes/MovieGraphGlobalOutputSettingNode.h
+++ b/Engine/Plugins/MovieScene/MovieRenderPipeline/Source/MovieRenderPipelineCore/Public/Graph/Nodes/MovieGraphGlobalOutputSettingNode.h
@@ -132,6 +132,7 @@ public:
 	UE_API virtual FSlateIcon GetIconAndTint(FLinearColor& OutColor) const override;
 	UE_API virtual EMovieGraphBranchRestriction GetBranchRestriction() const override;
 #endif
+	UE_API virtual bool GetCustomPromotionDefaultValue(const FMovieGraphPropertyInfo& InPropertyInfo, FString& OutSerializedValue) const override;
 	// ~UMovieGraphSettingNode Interface
 	
 protected:
diff --git a/Engine/Plugins/MovieScene/MovieRenderPipeline/Source/MovieRenderPipelineEditor/Private/Graph/MovieEdGraphNode.cpp b/Engine/Plugins/MovieScene/MovieRenderPipeline/Source/MovieRenderPipelineEditor/Private/Graph/MovieEdGraphNode.cpp
--- a/Engine/Plugins/MovieScene/MovieRenderPipeline/Source/MovieRenderPipelineEditor/Private/Graph/MovieEdGraphNode.cpp
+++ b/Engine/Plugins/MovieScene/MovieRenderPipeline/Source/MovieRenderPipelineEditor/Private/Graph/MovieEdGraphNode.cpp
@@ -423,19 +423,23 @@ void UMoviePipelineEdGraphNode::PromotePropertyToVariable(const FMovieGraphPrope
 		UObject* ValueTypeObject = const_cast<UObject*>(TargetProperty.ValueTypeObject.Get());
 		NewGraphVariable->SetValueType(TargetProperty.ValueType, ValueTypeObject);
 
-		// When promoting, set the variable's default value to the connected property's current value. That will ensure that there's no
-		// behavior change in the graph after the promotion.
+		// When promoting, set the variable's default value. The node can optionally provide a custom
+		// default value (eg, for properties that have special behavior when not overridden). Otherwise,
+		// use the property's current value so there's no behavior change in the graph after the promotion.
 		{
 			FString TargetPropertyValue;
-			if (TargetProperty.bIsDynamicProperty)
+			if (!RuntimeNode->GetCustomPromotionDefaultValue(TargetProperty, TargetPropertyValue))
 			{
-				RuntimeNode->GetDynamicPropertyValue(TargetProperty.Name, TargetPropertyValue); 
-			}
-			else if (const FProperty* TargetFProperty = FindFProperty<FProperty>(RuntimeNode->GetClass(), TargetProperty.Name))
-			{
-				TargetFProperty->ExportTextItem_InContainer(TargetPropertyValue, RuntimeNode, nullptr, RuntimeNode, PPF_None);
+				if (TargetProperty.bIsDynamicProperty)
+				{
+					RuntimeNode->GetDynamicPropertyValue(TargetProperty.Name, TargetPropertyValue);
+				}
+				else if (const FProperty* TargetFProperty = FindFProperty<FProperty>(RuntimeNode->GetClass(), TargetProperty.Name))
+				{
+					TargetFProperty->ExportTextItem_InContainer(TargetPropertyValue, RuntimeNode, nullptr, RuntimeNode, PPF_None);
+				}
 			}
-			
+
 			NewGraphVariable->SetValueSerializedString(TargetPropertyValue);
 		}
 
