MRG: Cache viewport camera at begin render instead of after PIE starts.

Prevents certain cases of the viewport not being active when that happens.

#rnx
#Synced-CL 52230590

diff --git a/Engine/Plugins/MovieScene/MovieRenderPipeline/Source/MovieRenderPipelineEditor/Private/Graph/MovieGraphQuickRender.cpp b/Engine/Plugins/MovieScene/MovieRenderPipeline/Source/MovieRenderPipelineEditor/Private/Graph/MovieGraphQuickRender.cpp
--- a/Engine/Plugins/MovieScene/MovieRenderPipeline/Source/MovieRenderPipelineEditor/Private/Graph/MovieGraphQuickRender.cpp
+++ b/Engine/Plugins/MovieScene/MovieRenderPipeline/Source/MovieRenderPipelineEditor/Private/Graph/MovieGraphQuickRender.cpp
@@ -925,13 +925,40 @@ void UMovieGraphQuickRenderSubsystem::CachePrePieData(UWorld* InEditorWorld)
 			CachedPrePieData.bUseViewportActorLockWorldTransform = true;
 			CachedPrePieData.ViewportActorLockWorldTransform = GetWorldSpaceCameraTransform(ActorLock.LockedActor.Get());
 		}
-		
+
 		// Duplicate the pilot cam component if there is one. As soon as PIE starts, the component's values will change to reflect the PIE world, so
 		// we duplicate it in order to be able to reference the pre-PIE property values.
 		if (const UCameraComponent* PilotCamComponent = ViewportClient->GetCameraComponentForView())
 		{
 			CachedPrePieData.ViewportActorLockCameraComponent.Reset(DuplicateObject<UCameraComponent>(PilotCamComponent, this));
 		}
+
+		// Cache the viewport camera properties now so we don't rely on the viewport being active when PIE starts later on.
+		CachedPrePieData.ViewportCamera.bIsValid = true;
+		CachedPrePieData.ViewportCamera.Location = ViewportClient->GetViewLocation();
+		CachedPrePieData.ViewportCamera.Rotation = ViewportClient->GetViewRotation();
+		CachedPrePieData.ViewportCamera.bIsPerspective = ViewportClient->IsPerspective();
+		CachedPrePieData.ViewportCamera.FOV = ViewportClient->ViewFOV;
+		CachedPrePieData.ViewportCamera.AspectRatio = ViewportClient->AspectRatio;
+		CachedPrePieData.ViewportCamera.OrthoWidth = ViewportClient->Viewport
+			? ViewportClient->Viewport->GetSizeXY().X * ViewportClient->GetOrthoUnitsPerPixel(ViewportClient->Viewport)
+			: 0.f;
+
+		// For axis-aligned ortho viewports, the rotation obtained from GetViewRotation() does not produce the correct
+		// result. Pre-compute the correct rotation here so the cached data is ready to use directly.
+		if (!CachedPrePieData.ViewportCamera.bIsPerspective)
+		{
+			switch (ViewportClient->ViewportType)
+			{
+			case LVT_OrthoTop:    CachedPrePieData.ViewportCamera.Rotation = FRotator(-90.f, -180.f, 0.f); break;
+			case LVT_OrthoBottom: CachedPrePieData.ViewportCamera.Rotation = FRotator(90.f, 0.f, 0.f); break;
+			case LVT_OrthoLeft:   CachedPrePieData.ViewportCamera.Rotation = FRotator(0.f, -90.f, 0.f); break;
+			case LVT_OrthoRight:  CachedPrePieData.ViewportCamera.Rotation = FRotator(0.f, 90.f, 0.f); break;
+			case LVT_OrthoBack:   CachedPrePieData.ViewportCamera.Rotation = FRotator(0.f, 0.f, 0.f); break;
+			case LVT_OrthoFront:  CachedPrePieData.ViewportCamera.Rotation = FRotator(0.f, 180.f, 0.f); break;
+			default:              CachedPrePieData.ViewportCamera.Rotation = FRotator(0.f, 0.f, 0.f); break;
+			}
+		}
 	}
 }
 
@@ -989,6 +1016,7 @@ void UMovieGraphQuickRenderSubsystem::RestorePreRenderState()
 	CachedPrePieData.ViewportActorLock.Reset();
 	CachedPrePieData.ViewportActorLockCameraComponent.Reset();
 	CachedPrePieData.bUseViewportActorLockWorldTransform = false;
+	CachedPrePieData.ViewportCamera.bIsValid = false;
 }
 
 ULevelSequence* UMovieGraphQuickRenderSubsystem::SetUpRenderingLevelSequence() const
@@ -1320,17 +1348,16 @@ void UMovieGraphQuickRenderSubsystem::PerformPreRenderSetup(UWorld* InEditorWorl
 		// a cine camera massively complicates things and creates a situation where it's difficult to match what the viewport looks like (eg, exposure).
 		ACameraActor* NewCamera = Cast<ACameraActor>(PieWorld->SpawnActor<ACameraActor>(CameraSpawnParams));
 
-		// Update the camera to look like the viewport's camera
-		if (const FLevelEditorViewportClient* ViewportClient = UMovieGraphApplyViewportLookNode::GetViewportClient())
+		// Update the camera to look like the viewport's camera. Use the viewport camera properties that were cached in
+		// CachePrePieData() rather than querying the viewport client now — by the time this lambda runs (after PIE
+		// starts), the viewport tab may no longer be active/visible, which would cause GetViewportClient() to return
+		// null and produce a render from the wrong camera.
+		if (CachedPrePieData.ViewportCamera.bIsValid)
 		{
-			// Initialize the camera properties to reflect the (non-pilot) view shown in the viewport 
-			FVector CameraPosition = ViewportClient->GetViewLocation();
-			FRotator CameraRotation = ViewportClient->GetViewRotation();
-			bool bIsPerspective = ViewportClient->IsPerspective();
-			float CameraFov = ViewportClient->ViewFOV;
-			float CameraAspectRatio = ViewportClient->AspectRatio;
-			float CameraOrthoWidth = ViewportClient->Viewport->GetSizeXY().X * ViewportClient->GetOrthoUnitsPerPixel(ViewportClient->Viewport);
-			
+			// Start from the cached viewport camera position/rotation (ortho rotation is already corrected at cache time)
+			FVector CameraPosition = CachedPrePieData.ViewportCamera.Location;
+			FRotator CameraRotation = CachedPrePieData.ViewportCamera.Rotation;
+
 			// If the viewport is being piloted, then we should use the properties of that to render from instead of the viewport's
 			// non-piloted camera properties. Note that this only covers non-camera pilot actors; camera pilot actors are covered above.
 			if (CachedPrePieData.bUseViewportActorLockWorldTransform)
@@ -1338,54 +1365,20 @@ void UMovieGraphQuickRenderSubsystem::PerformPreRenderSetup(UWorld* InEditorWorl
 				CameraPosition = CachedPrePieData.ViewportActorLockWorldTransform.GetTranslation();
 				CameraRotation = CachedPrePieData.ViewportActorLockWorldTransform.GetRotation().Rotator();
 			}
-			else
-			{
-				// The camera properties were initialized from the viewport camera properties, so all we need to do here is take into account the
-				// special case of the viewport showing an axis-aligned ortho cam. Getting the actor rotation from the viewport client for these
-				// ortho cams does not work properly, so their rotation is manually specified here. Sharing this logic with the viewport would involve
-				// some fairly significant changes, so for now MRG calculates this separately.
-				if (!bIsPerspective)
-				{
-					switch (ViewportClient->ViewportType)
-					{
-					case LVT_OrthoTop:
-						CameraRotation = FRotator(-90.0f, -180.0f, 0.0f);
-						break;
-					case LVT_OrthoBottom:
-						CameraRotation = FRotator(90.0f, 0.0f, 0.0f);
-						break;
-					case LVT_OrthoLeft:
-						CameraRotation = FRotator(0.0f, -90.0f, 0.0f);
-						break;
-					case LVT_OrthoRight:
-						CameraRotation = FRotator(0.0f, 90.0f, 0.0f);
-						break;
-					case LVT_OrthoBack:
-						CameraRotation = FRotator(0.0f, 0.0f, 0.0f);
-						break;
-					case LVT_OrthoFront:
-						CameraRotation = FRotator(0.0f, 180.0f, 0.0f);
-						break;
-					default:
-						CameraRotation = FRotator(0.0f, 0.0f, 0.0f);
-					}
-				}
-			}
 
 			NewCamera->SetActorLocation(CameraPosition);
 			NewCamera->SetActorRotation(CameraRotation);
 
 			// Note: At extreme zoom levels, the ortho clipping planes don't match up exactly with the viewport. This is something to improve in the future.
 			UCameraComponent* CameraComponent = NewCamera->GetCameraComponent();
-			CameraComponent->SetProjectionMode(bIsPerspective ? ECameraProjectionMode::Perspective : ECameraProjectionMode::Orthographic);
-			CameraComponent->SetAspectRatio(CameraAspectRatio);
-			CameraComponent->SetFieldOfView(CameraFov);
-			CameraComponent->SetOrthoWidth(CameraOrthoWidth);
+			CameraComponent->SetProjectionMode(CachedPrePieData.ViewportCamera.bIsPerspective ? ECameraProjectionMode::Perspective : ECameraProjectionMode::Orthographic);
+			CameraComponent->SetAspectRatio(CachedPrePieData.ViewportCamera.AspectRatio);
+			CameraComponent->SetFieldOfView(CachedPrePieData.ViewportCamera.FOV);
+			CameraComponent->SetOrthoWidth(CachedPrePieData.ViewportCamera.OrthoWidth);
 			CameraComponent->SetAutoCalculateOrthoPlanes(true);
 		}
 		else
 		{
-			// This case should be exceedingly rare, but log about it anyway
 			UE_LOG(LogMovieRenderPipeline, Error, TEXT("Quick Render: Could not determine the active viewport to render from, so the camera used will be incorrect."));
 		}
 
diff --git a/Engine/Plugins/MovieScene/MovieRenderPipeline/Source/MovieRenderPipelineEditor/Public/Graph/MovieGraphQuickRender.h b/Engine/Plugins/MovieScene/MovieRenderPipeline/Source/MovieRenderPipelineEditor/Public/Graph/MovieGraphQuickRender.h
--- a/Engine/Plugins/MovieScene/MovieRenderPipeline/Source/MovieRenderPipelineEditor/Public/Graph/MovieGraphQuickRender.h
+++ b/Engine/Plugins/MovieScene/MovieRenderPipeline/Source/MovieRenderPipelineEditor/Public/Graph/MovieGraphQuickRender.h
@@ -82,6 +82,19 @@ private:
 		 * nullptr if the viewport is locked to a non-camera actor (like a light).
 		 */
 		TStrongObjectPtr<UCameraComponent> ViewportActorLockCameraComponent = nullptr;
+
+		/** Cached viewport camera properties for the non-piloted camera path. These are captured at BeginQuickRender time so the viewport
+		 *  does not need to be active/visible when the PIE world is set up later. */
+		struct FViewportCameraData
+		{
+			FVector Location = FVector::ZeroVector;
+			FRotator Rotation = FRotator::ZeroRotator;
+			float FOV = 90.f;
+			float AspectRatio = 1.78f;
+			float OrthoWidth = 0.f;
+			bool bIsPerspective = true;
+			bool bIsValid = false;
+		} ViewportCamera;
 	};
 	
 	/**
