MoviePipeline: Fixed an issue that could cause graphs to be marked as dirty as soon as a queue is opened.

#Synced-CL 52226548

diff --git a/Engine/Plugins/MovieScene/MovieRenderPipeline/Source/MovieRenderPipelineCore/Public/Graph/MovieGraphMultiValueContainer.cpp b/Engine/Plugins/MovieScene/MovieRenderPipeline/Source/MovieRenderPipelineCore/Public/Graph/MovieGraphMultiValueContainer.cpp
--- a/Engine/Plugins/MovieScene/MovieRenderPipeline/Source/MovieRenderPipelineCore/Public/Graph/MovieGraphMultiValueContainer.cpp
+++ b/Engine/Plugins/MovieScene/MovieRenderPipeline/Source/MovieRenderPipelineCore/Public/Graph/MovieGraphMultiValueContainer.cpp
@@ -197,6 +197,12 @@ void UMovieGraphMultiValueContainer::SetValueType(const FName& InPropertyName, E
 {
 	if (ValuesView->HasProperty(InPropertyName))
 	{
+		// Skip if the type already matches to avoid unnecessary property change events (which can dirty packages)
+		if ((ValuesView->GetValueType(InPropertyName) == ValueType) && (ValuesView->GetValueTypeObject(InPropertyName) == InValueTypeObject))
+		{
+			return;
+		}
+
 		ValuesView->SetValueType(InPropertyName, ValueType, InValueTypeObject);
 
 #if WITH_EDITOR
