MoviePipeline: Fixed an issue that could cause an editor crash if a new (unsaved) graph is assigned to a job, the graph editor is opened, then an undo is performed.

#Synced-CL 52226548

diff --git a/Engine/Plugins/MovieScene/MovieRenderPipeline/Source/MovieRenderPipelineEditor/Private/Widgets/Graph/SMovieGraphConfigPanel.cpp b/Engine/Plugins/MovieScene/MovieRenderPipeline/Source/MovieRenderPipelineEditor/Private/Widgets/Graph/SMovieGraphConfigPanel.cpp
--- a/Engine/Plugins/MovieScene/MovieRenderPipeline/Source/MovieRenderPipelineEditor/Private/Widgets/Graph/SMovieGraphConfigPanel.cpp
+++ b/Engine/Plugins/MovieScene/MovieRenderPipeline/Source/MovieRenderPipelineEditor/Private/Widgets/Graph/SMovieGraphConfigPanel.cpp
@@ -699,6 +699,14 @@ void SMoviePipelineGraphPanel::PostUndo(bool bSuccess)
 
 		if (GraphEditorWidget.IsValid())
 		{
+			// Undo has the potential to leave behind an invalid graph or schema. To be on the safe side, close out
+			// the editor if this is the case.
+			if (!GraphEditorWidget->GetCurrentGraph() || !GraphEditorWidget->GetCurrentGraph()->GetSchema())
+			{
+				GEditor->GetEditorSubsystem<UAssetEditorSubsystem>()->CloseAllEditorsForAsset(CurrentGraph);
+				return;
+			}
+
 			// This unfortunately causes the graph to flicker for one frame, but all Blueprint Graphs do this.
 			// This is needed so that when a node is pasted, then undone, the widgets on the graph for the now
 			// deleted nodes also get removed, otherwise we try to draw them and they're pointed to invalid uobjects.
