MoviePipeline: Speculative fix for a rare crash in MRQ/MRG that could happen if the pipeline is torn down before a render has a chance to start (which could happen if there's a PIE initialization error).

#Synced-CL 52226548

diff --git a/Engine/Plugins/MovieScene/MovieRenderPipeline/Source/MovieRenderPipelineCore/Private/Graph/MovieGraphPipeline.cpp b/Engine/Plugins/MovieScene/MovieRenderPipeline/Source/MovieRenderPipelineCore/Private/Graph/MovieGraphPipeline.cpp
--- a/Engine/Plugins/MovieScene/MovieRenderPipeline/Source/MovieRenderPipelineCore/Private/Graph/MovieGraphPipeline.cpp
+++ b/Engine/Plugins/MovieScene/MovieRenderPipeline/Source/MovieRenderPipelineCore/Private/Graph/MovieGraphPipeline.cpp
@@ -1298,7 +1298,11 @@ void UMovieGraphPipeline::TransitionToState(const EMovieRenderPipelineState InNe
 
 			// If we had naturally finished the last shot before doing this transition it will have
 			// already been torn down, so this only catches mid-shot transitions to ensure teardown.
-			if (bShutdownRequested && (CurrentShotIndex < ActiveShotList.Num()))
+			// We also check that the shot was actually initialized to avoid restoring hierarchy state
+			// that was never saved. This can happen when PIE ends on the same frame as pipeline
+			// initialization, before the first tick runs InitializeShot/SetSoloShot.
+			if (bShutdownRequested && (CurrentShotIndex < ActiveShotList.Num())
+				&& ActiveShotList[CurrentShotIndex]->ShotInfo.IsInitialized())
 			{
 				// Ensures all in-flight work for that shot is handled.
 				TeardownShot(ActiveShotList[CurrentShotIndex]);
@@ -1316,8 +1320,12 @@ void UMovieGraphPipeline::TransitionToState(const EMovieRenderPipelineState InNe
 
 			// Restore any custom Time Step that may have been set before. We do this here
 			// because the TimeStepInstance is only expected to be having to calculate times
-			// during ProducingFrames.
-			GetTimeStepInstance()->Shutdown();
+			// during ProducingFrames. The instance may be null if the pipeline was shut down
+			// before the first TickProducingFrames assigned it.
+			if (GetTimeStepInstance())
+			{
+				GetTimeStepInstance()->Shutdown();
+			}
 
 			// We assign it to nullptr here so that post-ProducingFrame code doesn't accidentally 
 			// use a stale time step instance.
diff --git a/Engine/Plugins/MovieScene/MovieRenderPipeline/Source/MovieRenderPipelineCore/Private/MoviePipeline.cpp b/Engine/Plugins/MovieScene/MovieRenderPipeline/Source/MovieRenderPipelineCore/Private/MoviePipeline.cpp
--- a/Engine/Plugins/MovieScene/MovieRenderPipeline/Source/MovieRenderPipelineCore/Private/MoviePipeline.cpp
+++ b/Engine/Plugins/MovieScene/MovieRenderPipeline/Source/MovieRenderPipelineCore/Private/MoviePipeline.cpp
@@ -463,7 +463,11 @@ void UMoviePipeline::TransitionToState(const EMovieRenderPipelineState InNewStat
 
 			// If we had naturally finished the last shot before doing this transition it will have
 			// already been torn down, so this only catches mid-shot transitions to ensure teardown.
-			if (CurrentShotIndex < ActiveShotList.Num())
+			// We also check that the shot was actually initialized (State != Uninitialized) to avoid
+			// restoring hierarchy state that was never saved. This can happen when PIE ends on the same
+			// frame as MRQ initialization, before the first TickProducingFrames/InitializeShot call.
+			if (CurrentShotIndex < ActiveShotList.Num()
+				&& ActiveShotList[CurrentShotIndex]->ShotInfo.IsInitialized())
 			{
 				// Ensures all in-flight work for that shot is handled.
 				TeardownShot(ActiveShotList[CurrentShotIndex]);
diff --git a/Engine/Plugins/MovieScene/MovieRenderPipeline/Source/MovieRenderPipelineCore/Private/MoviePipelineUtils.cpp b/Engine/Plugins/MovieScene/MovieRenderPipeline/Source/MovieRenderPipelineCore/Private/MoviePipelineUtils.cpp
--- a/Engine/Plugins/MovieScene/MovieRenderPipeline/Source/MovieRenderPipelineCore/Private/MoviePipelineUtils.cpp
+++ b/Engine/Plugins/MovieScene/MovieRenderPipeline/Source/MovieRenderPipelineCore/Private/MoviePipelineUtils.cpp
@@ -480,6 +480,17 @@ namespace MoviePipeline
 		TSharedPtr<FCameraCutSubSectionHierarchyNode> Node = InLeaf;
 		while (Node)
 		{
+			// Use the playback range as a sentinel: if it has open bounds, save was never called
+			// for this node (e.g. the pipeline shut down before hierarchy was saved). Skip the
+			// entire restoration to avoid writing default/uninitialized values back.
+			if (!bInSave && Node->MovieScene.IsValid()
+				&& !ensureMsgf(Node->OriginalMovieScenePlaybackRange.GetLowerBound().IsClosed() && Node->OriginalMovieScenePlaybackRange.GetUpperBound().IsClosed(),
+					TEXT("Skipping restoration for MovieScene '%s' — cached range has open bounds (pipeline shut down before hierarchy was saved?)."), *GetNameSafe(Node->MovieScene.Get())))
+			{
+				Node = Node->GetParent();
+				continue;
+			}
+
 			if (Node->MovieScene.IsValid())
 			{
 				if (bInSave)
