MoviePipeline: Fixed log spam that can occur when opening the Quick Render settings drop-down from the viewport toolbar.

#Synced-CL 52313960

diff --git a/Engine/Plugins/MovieScene/MovieRenderPipeline/Source/MovieRenderPipelineCore/Private/Graph/MovieGraphQuickRenderSettings.cpp b/Engine/Plugins/MovieScene/MovieRenderPipeline/Source/MovieRenderPipelineCore/Private/Graph/MovieGraphQuickRenderSettings.cpp
--- a/Engine/Plugins/MovieScene/MovieRenderPipeline/Source/MovieRenderPipelineCore/Private/Graph/MovieGraphQuickRenderSettings.cpp
+++ b/Engine/Plugins/MovieScene/MovieRenderPipeline/Source/MovieRenderPipelineCore/Private/Graph/MovieGraphQuickRenderSettings.cpp
@@ -27,8 +27,15 @@ UMovieGraphQuickRenderModeSettings* UMovieGraphQuickRenderSettings::GetSavedQuic
 	};
 	
 	UMovieGraphQuickRenderSettings* QuickRenderSettings = nullptr;
-	
-	if (const UPackage* QuickRenderSettingsPackage = LoadPackage(nullptr, *QuickRenderSettingsPackagePath, LOAD_None))
+
+	// Use FindPackage first to avoid triggering the async loading path on packages that are already in memory.
+	const UPackage* QuickRenderSettingsPackage = FindPackage(nullptr, *QuickRenderSettingsPackagePath);
+	if (!QuickRenderSettingsPackage)
+	{
+		QuickRenderSettingsPackage = LoadPackage(nullptr, *QuickRenderSettingsPackagePath, LOAD_None);
+	}
+
+	if (QuickRenderSettingsPackage)
 	{
 		QuickRenderSettings = Cast<UMovieGraphQuickRenderSettings>(FindObjectWithOuter(QuickRenderSettingsPackage, StaticClass()));
 	}
diff --git a/Engine/Plugins/MovieScene/MovieRenderPipeline/Source/MovieRenderPipelineEditor/Private/MoviePipelinePostRenderFileDisplayProcessor.cpp b/Engine/Plugins/MovieScene/MovieRenderPipeline/Source/MovieRenderPipelineEditor/Private/MoviePipelinePostRenderFileDisplayProcessor.cpp
--- a/Engine/Plugins/MovieScene/MovieRenderPipeline/Source/MovieRenderPipelineEditor/Private/MoviePipelinePostRenderFileDisplayProcessor.cpp
+++ b/Engine/Plugins/MovieScene/MovieRenderPipeline/Source/MovieRenderPipelineEditor/Private/MoviePipelinePostRenderFileDisplayProcessor.cpp
@@ -774,7 +774,14 @@ UMoviePipelineOutputDataContainer* FMoviePipelinePostRenderFileDisplayProcessor:
 {
 	UMoviePipelineOutputDataContainer* OutputData = nullptr;
 
-	if (const UPackage* OutputDataPackage = LoadPackage(nullptr, *PersistedOutputDataPackagePath, LOAD_None))
+	// Use FindPackage first to avoid triggering the async loading path on packages that are already in memory.
+	const UPackage* OutputDataPackage = FindPackage(nullptr, *PersistedOutputDataPackagePath);
+	if (!OutputDataPackage)
+	{
+		OutputDataPackage = LoadPackage(nullptr, *PersistedOutputDataPackagePath, LOAD_None);
+	}
+
+	if (OutputDataPackage)
 	{
 		OutputData = Cast<UMoviePipelineOutputDataContainer>(FindObjectWithOuter(OutputDataPackage, UMoviePipelineOutputDataContainer::StaticClass()));
 	}
