MoviePipeline: Added two new UMovieGraphConfig API methods: InsertBefore() and InsertAfter() insert nodes either upstream (to the left) or downstream (to the right) of a specific node.

#Synced-CL 52423999

diff --git a/Engine/Plugins/MovieScene/MovieRenderPipeline/Source/MovieRenderPipelineCore/Private/Graph/MovieGraphConfig.cpp b/Engine/Plugins/MovieScene/MovieRenderPipeline/Source/MovieRenderPipelineCore/Private/Graph/MovieGraphConfig.cpp
--- a/Engine/Plugins/MovieScene/MovieRenderPipeline/Source/MovieRenderPipelineCore/Private/Graph/MovieGraphConfig.cpp
+++ b/Engine/Plugins/MovieScene/MovieRenderPipeline/Source/MovieRenderPipelineCore/Private/Graph/MovieGraphConfig.cpp
@@ -815,6 +815,113 @@ bool UMovieGraphConfig::RemoveOutboundEdges(UMovieGraphNode* InNode, const FName
 	return bChanged;
 }
 
+UMovieGraphNode* UMovieGraphConfig::InsertBefore(UMovieGraphNode* InTargetNode, const TSubclassOf<UMovieGraphNode> InNodeClass, const FName& InTargetPinLabel)
+{
+	if (!IsValid(InTargetNode))
+	{
+		UE_LOG(LogMovieRenderPipeline, Warning, TEXT("[%hs]: InTargetNode is not valid."), __FUNCTION__);
+		return nullptr;
+	}
+
+	if (!InNodeClass)
+	{
+		UE_LOG(LogMovieRenderPipeline, Warning, TEXT("[%hs]: InNodeClass is null."), __FUNCTION__);
+		return nullptr;
+	}
+
+	if (!InTargetNode->GetInputPin(InTargetPinLabel))
+	{
+		UE_LOG(LogMovieRenderPipeline, Warning, TEXT("[%hs]: Target node '%s' (%s) has no input pin '%s'."),
+			__FUNCTION__, *InTargetNode->GetName(), *InTargetNode->GetClass()->GetName(), *InTargetPinLabel.ToString());
+		return nullptr;
+	}
+
+	// Save info about the node currently upstream of the target before the connection is broken by AddLabeledEdge
+	UMovieGraphNode* UpstreamNode = nullptr;
+	FName UpstreamOutputPinLabel;
+	if (const UMovieGraphPin* ConnectedPin = InTargetNode->GetInputPin(InTargetPinLabel)->GetFirstConnectedPin())
+	{
+		UpstreamNode = ConnectedPin->Node;
+		UpstreamOutputPinLabel = ConnectedPin->Properties.Label;
+	}
+
+	Modify();
+
+	UMovieGraphNode* NewNode = CreateNodeByClass(InNodeClass);
+	if (!NewNode)
+	{
+		return nullptr;
+	}
+
+	// Connect the new node to the target; input pins only allow one connection, so this auto-breaks
+	// the existing UpstreamNode → InTargetNode edge
+	AddLabeledEdge(NewNode, FName(), InTargetNode, InTargetPinLabel);
+
+	// Reconnect the previously upstream node to the new node's input
+	if (UpstreamNode)
+	{
+		AddLabeledEdge(UpstreamNode, UpstreamOutputPinLabel, NewNode, FName());
+	}
+
+	return NewNode;
+}
+
+UMovieGraphNode* UMovieGraphConfig::InsertAfter(UMovieGraphNode* InTargetNode, const TSubclassOf<UMovieGraphNode> InNodeClass, const FName& InTargetPinLabel)
+{
+	if (!IsValid(InTargetNode))
+	{
+		UE_LOG(LogMovieRenderPipeline, Warning, TEXT("[%hs]: InTargetNode is not valid."), __FUNCTION__);
+		return nullptr;
+	}
+
+	if (!InNodeClass)
+	{
+		UE_LOG(LogMovieRenderPipeline, Warning, TEXT("[%hs]: InNodeClass is null."), __FUNCTION__);
+		return nullptr;
+	}
+
+	if (!InTargetNode->GetOutputPin(InTargetPinLabel))
+	{
+		UE_LOG(LogMovieRenderPipeline, Warning, TEXT("[%hs]: Target node '%s' (%s) has no output pin '%s'."),
+			__FUNCTION__, *InTargetNode->GetName(), *InTargetNode->GetClass()->GetName(), *InTargetPinLabel.ToString());
+		return nullptr;
+	}
+
+	// Save info about the node currently downstream of the target so it can be reconnected after insertion.
+	// The downstream pin's label is preserved so named branch pins (e.g. on the output node) are handled correctly.
+	UMovieGraphNode* DownstreamNode = nullptr;
+	FName DownstreamInputPinLabel;
+	if (const UMovieGraphPin* ConnectedPin = InTargetNode->GetOutputPin(InTargetPinLabel)->GetFirstConnectedPin())
+	{
+		DownstreamNode = ConnectedPin->Node;
+		DownstreamInputPinLabel = ConnectedPin->Properties.Label;
+	}
+
+	Modify();
+
+	UMovieGraphNode* NewNode = CreateNodeByClass(InNodeClass);
+	if (!NewNode)
+	{
+		return nullptr;
+	}
+
+	// Explicitly break the existing InTargetNode → DownstreamNode edge before reconnecting; output pins allow
+	// multiple connections so AddLabeledEdge alone will not remove the old edge
+	if (DownstreamNode)
+	{
+		RemoveLabeledEdge(InTargetNode, InTargetPinLabel, DownstreamNode, DownstreamInputPinLabel);
+	}
+
+	// Connect the target to the new node, then reconnect the new node to what was previously downstream
+	AddLabeledEdge(InTargetNode, InTargetPinLabel, NewNode, FName());
+	if (DownstreamNode)
+	{
+		AddLabeledEdge(NewNode, FName(), DownstreamNode, DownstreamInputPinLabel);
+	}
+
+	return NewNode;
+}
+
 void UMovieGraphConfig::AddNode(UMovieGraphNode* InNode)
 {
 	if (!InNode)
diff --git a/Engine/Plugins/MovieScene/MovieRenderPipeline/Source/MovieRenderPipelineCore/Public/Graph/MovieGraphConfig.h b/Engine/Plugins/MovieScene/MovieRenderPipeline/Source/MovieRenderPipelineCore/Public/Graph/MovieGraphConfig.h
--- a/Engine/Plugins/MovieScene/MovieRenderPipeline/Source/MovieRenderPipelineCore/Public/Graph/MovieGraphConfig.h
+++ b/Engine/Plugins/MovieScene/MovieRenderPipeline/Source/MovieRenderPipelineCore/Public/Graph/MovieGraphConfig.h
@@ -696,6 +696,23 @@ public:
 	UE_API bool RemoveOutboundEdges(UMovieGraphNode* InNode, const FName& InPinName);
 
 	/** 
+	 * Creates a new node of the given class and inserts it immediately upstream (to the left) of InTargetNode
+	 * on the specified input pin (defaults to the unnamed pin). Any node previously connected to that input
+	 * is reconnected upstream of the new node. Returns the newly created node, or nullptr on failure.
+	 */
+	UFUNCTION(BlueprintCallable, Category = "Movie Graph")
+	UE_API UMovieGraphNode* InsertBefore(UMovieGraphNode* InTargetNode, const TSubclassOf<UMovieGraphNode> InNodeClass, const FName& InTargetPinLabel = FName());
+
+	/**
+	 * Creates a new node of the given class and inserts it immediately downstream (to the right) of InTargetNode
+	 * on the specified output pin (defaults to the unnamed pin). The first downstream connection on that pin is
+	 * reconnected downstream of the new node; if the pin has multiple connections, the caller should manage
+	 * additional reconnections manually. Returns the newly created node, or nullptr on failure.
+	 */
+	UFUNCTION(BlueprintCallable, Category = "Movie Graph")
+	UE_API UMovieGraphNode* InsertAfter(UMovieGraphNode* InTargetNode, const TSubclassOf<UMovieGraphNode> InNodeClass, const FName& InTargetPinLabel = FName());
+
+	/**
 	* Add the specified node instance to the graph. This will rename the node to ensure the graph is the outer
 	* and then it will add it to the internal list of nodes used by the graph. See ConstructRuntimeNode if you
 	* want to construct a node by class and don't already have an instance.
@@ -981,6 +998,10 @@ public:
 	//
 	// Note that the editor graph is saved with the runtime graph. This is done to prevent the runtime graph from being dirtied immediately upon loading
 	// (because the editor graph would have to be re-created from the runtime graph, thus dirtying the package).
+	//
+	// If the runtime graph representation is modified outside of the editor (eg, programmatically via the UMovieGraphConfig API
+	// directly), this should be nulled out so the editor can rebuild its representation the next time it is loaded. If the editor representation
+	// of the graph isn't important (ie, the graph is transient) then this is unnecessary.
 	UPROPERTY()
 	TObjectPtr<UEdGraph> PipelineEdGraph;
 #endif
diff --git a/Engine/Plugins/MovieScene/MovieRenderPipeline/Source/MovieRenderPipelineEditor/Private/Graph/MovieEdGraph.cpp b/Engine/Plugins/MovieScene/MovieRenderPipeline/Source/MovieRenderPipelineEditor/Private/Graph/MovieEdGraph.cpp
--- a/Engine/Plugins/MovieScene/MovieRenderPipeline/Source/MovieRenderPipelineEditor/Private/Graph/MovieEdGraph.cpp
+++ b/Engine/Plugins/MovieScene/MovieRenderPipeline/Source/MovieRenderPipelineEditor/Private/Graph/MovieEdGraph.cpp
@@ -33,6 +33,25 @@ void UMoviePipelineEdGraph::InitFromRuntimeGraph(UMovieGraphConfig* InGraph)
 {
 	check(InGraph);
 
+	// Runtime nodes may carry stale GraphNode pointers from a previous editor graph (e.g., when a graph was
+	// programmatically modified via UMovieGraphConfig APIs after being duplicated from a saved asset, then had
+	// its PipelineEdGraph reset to null). Clear any GraphNode that doesn't belong to this editor graph so fresh
+	// editor nodes are created below. This ensures the editor representation always matches the runtime state.
+	auto ClearStaleGraphNode = [this](UMovieGraphNode* RuntimeNode)
+	{
+		if (RuntimeNode && RuntimeNode->GraphNode
+			&& (!IsValid(RuntimeNode->GraphNode) || RuntimeNode->GraphNode->GetOuter() != this))
+		{
+			RuntimeNode->GraphNode = nullptr;
+		}
+	};
+	ClearStaleGraphNode(InGraph->GetInputNode());
+	ClearStaleGraphNode(InGraph->GetOutputNode());
+	for (const TObjectPtr<UMovieGraphNode>& RuntimeNode : InGraph->GetNodes())
+	{
+		ClearStaleGraphNode(RuntimeNode.Get());
+	}
+
 	TMap<UMovieGraphNode*, UMoviePipelineEdGraphNodeBase*> NodeLookup;
 
 	// Input
diff --git a/Engine/Plugins/MovieScene/MovieRenderPipeline/Source/MovieRenderPipelineEditor/Private/Graph/MovieGraphQuickRender.cpp b/Engine/Plugins/MovieScene/MovieRenderPipeline/Source/MovieRenderPipelineEditor/Private/Graph/MovieGraphQuickRender.cpp
--- a/Engine/Plugins/MovieScene/MovieRenderPipeline/Source/MovieRenderPipelineEditor/Private/Graph/MovieGraphQuickRender.cpp
+++ b/Engine/Plugins/MovieScene/MovieRenderPipeline/Source/MovieRenderPipelineEditor/Private/Graph/MovieGraphQuickRender.cpp
@@ -843,33 +843,16 @@ UMovieGraphNode* UMovieGraphQuickRenderSubsystem::InjectNodeIntoBranch(const TSu
 {
 	check(IsValid(TemporaryGraph));
 
-	UMovieGraphNode* NewNode = TemporaryGraph->CreateNodeByClass(NodeType);
 	UMovieGraphNode* OutputNode = TemporaryGraph->GetOutputNode();
-
-	const UMovieGraphPin* BranchPin = OutputNode->GetInputPin(InBranchName);
-	if (!BranchPin)
+	if (!OutputNode->GetInputPin(InBranchName))
 	{
 		UE_LOG(LogMovieRenderPipeline, Warning, TEXT("Quick Render: Could not find branch [%s] to inject a node override into. The render may look different than expected."), *InBranchName.ToString());
 		return nullptr;
 	}
 
-	// Get the node on the branch that is most downstream (if any)
-	UMovieGraphNode* MostDownstreamNode = nullptr;
-	if (const UMovieGraphPin* ConnectedPin = BranchPin->GetFirstConnectedPin())
-	{
-		MostDownstreamNode = ConnectedPin->Node;
-	}
-
-	// Add in the new node to the branch, downstream from the existing node that was previously furthest downstream
-	TemporaryGraph->AddLabeledEdge(NewNode, FName(), OutputNode, InBranchName);
-
-	// Re-connect the previously connected node upstream of the new node, if there was one
-	if (MostDownstreamNode)
-	{
-		TemporaryGraph->AddLabeledEdge(MostDownstreamNode, FName(), NewNode, FName());
-	}
-
-	return NewNode;
+	// InsertBefore handles both cases: if a node is already connected to the branch pin it is reconnected
+	// upstream of the new node, and if the branch is empty the new node is simply connected to the pin.
+	return TemporaryGraph->InsertBefore(OutputNode, NodeType, InBranchName);
 }
 
 bool UMovieGraphQuickRenderSubsystem::AreModePrerequisitesMet() const
