MoviePipeline: Updated core MRG output nodes with a new GetBasicConfigShortDisplayName() method in preparation for the new "Basic" UI.

#Synced-CL 52423999

diff --git a/Engine/Plugins/Media/AppleProResMedia/Source/AppleProResMedia/Private/MovieGraphAppleProResNode.cpp b/Engine/Plugins/Media/AppleProResMedia/Source/AppleProResMedia/Private/MovieGraphAppleProResNode.cpp
--- a/Engine/Plugins/Media/AppleProResMedia/Source/AppleProResMedia/Private/MovieGraphAppleProResNode.cpp
+++ b/Engine/Plugins/Media/AppleProResMedia/Source/AppleProResMedia/Private/MovieGraphAppleProResNode.cpp
@@ -54,6 +54,11 @@ FSlateIcon UMovieGraphAppleProResNode::GetIconAndTint(FLinearColor& OutColor) co
 	OutColor = FLinearColor::White;
 	return AppleProResIcon;
 }
+
+FText UMovieGraphAppleProResNode::GetBasicConfigShortDisplayName() const
+{
+	return NSLOCTEXT("MovieGraphNodes", "BasicConfigShortName_PRORES", "PRORES");
+}
 #endif // WITH_EDITOR
 
 TUniquePtr<MovieRenderGraph::IVideoCodecWriter> UMovieGraphAppleProResNode::Initialize_GameThread(const FMovieGraphVideoNodeInitializationContext& InInitializationContext)
diff --git a/Engine/Plugins/Media/AppleProResMedia/Source/AppleProResMedia/Private/MovieGraphAppleProResNode.h b/Engine/Plugins/Media/AppleProResMedia/Source/AppleProResMedia/Private/MovieGraphAppleProResNode.h
--- a/Engine/Plugins/Media/AppleProResMedia/Source/AppleProResMedia/Private/MovieGraphAppleProResNode.h
+++ b/Engine/Plugins/Media/AppleProResMedia/Source/AppleProResMedia/Private/MovieGraphAppleProResNode.h
@@ -23,6 +23,7 @@ public:
 	virtual FText GetKeywords() const override;
 	virtual FLinearColor GetNodeTitleColor() const override;
 	virtual FSlateIcon GetIconAndTint(FLinearColor& OutColor) const override;
+	virtual FText GetBasicConfigShortDisplayName() const override;
 #endif
 
 protected:
diff --git a/Engine/Plugins/Media/AvidDNxMedia/Source/Source/Private/MovieGraphAvidDNxHRNode.cpp b/Engine/Plugins/Media/AvidDNxMedia/Source/Source/Private/MovieGraphAvidDNxHRNode.cpp
--- a/Engine/Plugins/Media/AvidDNxMedia/Source/Source/Private/MovieGraphAvidDNxHRNode.cpp
+++ b/Engine/Plugins/Media/AvidDNxMedia/Source/Source/Private/MovieGraphAvidDNxHRNode.cpp
@@ -50,6 +50,11 @@ FSlateIcon UMovieGraphAvidDNxHRNode::GetIconAndTint(FLinearColor& OutColor) cons
 	OutColor = FLinearColor::White;
 	return AvidDNxHRIcon;
 }
+
+FText UMovieGraphAvidDNxHRNode::GetBasicConfigShortDisplayName() const
+{
+	return NSLOCTEXT("MovieGraphNodes", "BasicConfigShortName_AVID", "AVID");
+}
 #endif // WITH_EDITOR
 
 TUniquePtr<MovieRenderGraph::IVideoCodecWriter> UMovieGraphAvidDNxHRNode::Initialize_GameThread(const FMovieGraphVideoNodeInitializationContext& InInitializationContext)
diff --git a/Engine/Plugins/Media/AvidDNxMedia/Source/Source/Private/MovieGraphAvidDNxHRNode.h b/Engine/Plugins/Media/AvidDNxMedia/Source/Source/Private/MovieGraphAvidDNxHRNode.h
--- a/Engine/Plugins/Media/AvidDNxMedia/Source/Source/Private/MovieGraphAvidDNxHRNode.h
+++ b/Engine/Plugins/Media/AvidDNxMedia/Source/Source/Private/MovieGraphAvidDNxHRNode.h
@@ -25,6 +25,10 @@ public:
 	virtual FText GetKeywords() const override;
 	virtual FLinearColor GetNodeTitleColor() const override;
 	virtual FSlateIcon GetIconAndTint(FLinearColor& OutColor) const override;
+
+	// UMovieGraphFileOutputNode Interface
+	virtual FText GetBasicConfigShortDisplayName() const override;
+	// ~UMovieGraphFileOutputNode Interface
 #endif
 
 protected:
diff --git a/Engine/Plugins/MovieScene/MovieRenderPipeline/Source/MovieRenderPipelineCore/Private/Graph/Nodes/MovieGraphAudioOutputNode.cpp b/Engine/Plugins/MovieScene/MovieRenderPipeline/Source/MovieRenderPipelineCore/Private/Graph/Nodes/MovieGraphAudioOutputNode.cpp
--- a/Engine/Plugins/MovieScene/MovieRenderPipeline/Source/MovieRenderPipelineCore/Private/Graph/Nodes/MovieGraphAudioOutputNode.cpp
+++ b/Engine/Plugins/MovieScene/MovieRenderPipeline/Source/MovieRenderPipelineCore/Private/Graph/Nodes/MovieGraphAudioOutputNode.cpp
@@ -64,6 +64,11 @@ FSlateIcon UMovieGraphAudioOutputNode::GetIconAndTint(FLinearColor& OutColor) co
 	OutColor = FLinearColor::White;
 	return AudioOutputPresetIcon;
 }
+
+FText UMovieGraphAudioOutputNode::GetBasicConfigShortDisplayName() const
+{
+	return NSLOCTEXT("MovieGraphNodes", "BasicConfigShortName_WAV", "WAV");
+}
 #endif // WITH_EDITOR
 
 EMovieGraphBranchRestriction UMovieGraphAudioOutputNode::GetBranchRestriction() const
diff --git a/Engine/Plugins/MovieScene/MovieRenderPipeline/Source/MovieRenderPipelineCore/Private/Graph/Nodes/MovieGraphFileOutputNode.cpp b/Engine/Plugins/MovieScene/MovieRenderPipeline/Source/MovieRenderPipelineCore/Private/Graph/Nodes/MovieGraphFileOutputNode.cpp
--- a/Engine/Plugins/MovieScene/MovieRenderPipeline/Source/MovieRenderPipelineCore/Private/Graph/Nodes/MovieGraphFileOutputNode.cpp
+++ b/Engine/Plugins/MovieScene/MovieRenderPipeline/Source/MovieRenderPipelineCore/Private/Graph/Nodes/MovieGraphFileOutputNode.cpp
@@ -9,6 +9,13 @@
 
 #include UE_INLINE_GENERATED_CPP_BY_NAME(MovieGraphFileOutputNode)
 
+#if WITH_EDITOR
+FText UMovieGraphFileOutputNode::GetBasicConfigShortDisplayName() const
+{
+	return FText::GetEmpty();
+}
+#endif	// WITH_EDITOR
+
 int32 UMovieGraphFileOutputNode::GetNumFileOutputNodes(const UMovieGraphEvaluatedConfig& InEvaluatedConfig, const FName& InBranchName)
 {
 	return InEvaluatedConfig.GetSettingsForBranch(UMovieGraphFileOutputNode::StaticClass(), InBranchName, false /*bIncludeCDOs*/, false /*bExactMatch*/).Num();
diff --git a/Engine/Plugins/MovieScene/MovieRenderPipeline/Source/MovieRenderPipelineCore/Public/Graph/Nodes/MovieGraphAudioOutputNode.h b/Engine/Plugins/MovieScene/MovieRenderPipeline/Source/MovieRenderPipelineCore/Public/Graph/Nodes/MovieGraphAudioOutputNode.h
--- a/Engine/Plugins/MovieScene/MovieRenderPipeline/Source/MovieRenderPipelineCore/Public/Graph/Nodes/MovieGraphAudioOutputNode.h
+++ b/Engine/Plugins/MovieScene/MovieRenderPipeline/Source/MovieRenderPipelineCore/Public/Graph/Nodes/MovieGraphAudioOutputNode.h
@@ -33,9 +33,12 @@ public:
 	UE_API virtual FText GetKeywords() const override;
 	UE_API virtual FLinearColor GetNodeTitleColor() const override;
 	UE_API virtual FSlateIcon GetIconAndTint(FLinearColor& OutColor) const override;
-#endif
 	// ~UMovieGraphNode Interface
 
+	// UMovieGraphFileOutputNode interface
+	UE_API virtual FText GetBasicConfigShortDisplayName() const override;
+#endif
+
 	/**
 	 * Although this node does not represent a proper "renderer", it will sometimes be identified in data structures
 	 * as a renderer. If that's the case, this is the name of the audio "renderer".
diff --git a/Engine/Plugins/MovieScene/MovieRenderPipeline/Source/MovieRenderPipelineCore/Public/Graph/Nodes/MovieGraphFileOutputNode.h b/Engine/Plugins/MovieScene/MovieRenderPipeline/Source/MovieRenderPipelineCore/Public/Graph/Nodes/MovieGraphFileOutputNode.h
--- a/Engine/Plugins/MovieScene/MovieRenderPipeline/Source/MovieRenderPipelineCore/Public/Graph/Nodes/MovieGraphFileOutputNode.h
+++ b/Engine/Plugins/MovieScene/MovieRenderPipeline/Source/MovieRenderPipelineCore/Public/Graph/Nodes/MovieGraphFileOutputNode.h
@@ -71,6 +71,12 @@ public:
 	{
 		return NSLOCTEXT("MovieGraphNodes", "FileOutputGraphNode_Category", "Output Type");
 	}
+
+	/**
+	 * Short display name for use in Basic config toggle buttons (e.g., "EXR", "PNG", etc).
+	 * Return empty FText to exclude this output type from the Basic config UI.
+	 */
+	UE_API virtual FText GetBasicConfigShortDisplayName() const;
 #endif
 	
 	UPROPERTY(EditAnywhere, BlueprintReadWrite, Category = Overrides, meta = (InlineEditConditionToggle))
diff --git a/Engine/Plugins/MovieScene/MovieRenderPipeline/Source/MovieRenderPipelineMP4Encoder/Private/Graph/MovieGraphMP4EncoderNode.cpp b/Engine/Plugins/MovieScene/MovieRenderPipeline/Source/MovieRenderPipelineMP4Encoder/Private/Graph/MovieGraphMP4EncoderNode.cpp
--- a/Engine/Plugins/MovieScene/MovieRenderPipeline/Source/MovieRenderPipelineMP4Encoder/Private/Graph/MovieGraphMP4EncoderNode.cpp
+++ b/Engine/Plugins/MovieScene/MovieRenderPipeline/Source/MovieRenderPipelineMP4Encoder/Private/Graph/MovieGraphMP4EncoderNode.cpp
@@ -53,6 +53,11 @@ FSlateIcon UMovieGraphMP4EncoderNode::GetIconAndTint(FLinearColor& OutColor) con
 	OutColor = FLinearColor::White;
 	return Icon;
 }
+
+FText UMovieGraphMP4EncoderNode::GetBasicConfigShortDisplayName() const
+{
+	return NSLOCTEXT("MovieGraphNodes", "BasicConfigShortName_MP4", "MP4");
+}
 #endif // WITH_EDITOR
 
 TUniquePtr<MovieRenderGraph::IVideoCodecWriter> UMovieGraphMP4EncoderNode::Initialize_GameThread(const FMovieGraphVideoNodeInitializationContext& InInitializationContext)
diff --git a/Engine/Plugins/MovieScene/MovieRenderPipeline/Source/MovieRenderPipelineMP4Encoder/Public/Graph/MovieGraphMP4EncoderNode.h b/Engine/Plugins/MovieScene/MovieRenderPipeline/Source/MovieRenderPipelineMP4Encoder/Public/Graph/MovieGraphMP4EncoderNode.h
--- a/Engine/Plugins/MovieScene/MovieRenderPipeline/Source/MovieRenderPipelineMP4Encoder/Public/Graph/MovieGraphMP4EncoderNode.h
+++ b/Engine/Plugins/MovieScene/MovieRenderPipeline/Source/MovieRenderPipelineMP4Encoder/Public/Graph/MovieGraphMP4EncoderNode.h
@@ -25,6 +25,7 @@ public:
 	UE_API virtual FText GetKeywords() const override;
 	UE_API virtual FLinearColor GetNodeTitleColor() const override;
 	UE_API virtual FSlateIcon GetIconAndTint(FLinearColor& OutColor) const override;
+	UE_API virtual FText GetBasicConfigShortDisplayName() const override;
 #endif
 
 protected:
diff --git a/Engine/Plugins/MovieScene/MovieRenderPipeline/Source/MovieRenderPipelineRenderPasses/Public/MovieGraphImageSequenceOutputNode.h b/Engine/Plugins/MovieScene/MovieRenderPipeline/Source/MovieRenderPipelineRenderPasses/Public/MovieGraphImageSequenceOutputNode.h
--- a/Engine/Plugins/MovieScene/MovieRenderPipeline/Source/MovieRenderPipelineRenderPasses/Public/MovieGraphImageSequenceOutputNode.h
+++ b/Engine/Plugins/MovieScene/MovieRenderPipeline/Source/MovieRenderPipelineRenderPasses/Public/MovieGraphImageSequenceOutputNode.h
@@ -264,6 +264,11 @@ public:
 		static const FText EXRSequenceNodeName = NSLOCTEXT("MovieGraph", "NodeName_EXRSequenceMultilayer", ".exr Sequence (Multilayer)");
 		return EXRSequenceNodeName;
 	}
+	
+	virtual FText GetBasicConfigShortDisplayName() const override
+	{
+		return NSLOCTEXT("MovieGraph", "BasicConfigShortName_EXR", "EXR");
+	}
 
 	virtual FText GetKeywords() const override
 	{
@@ -379,6 +384,11 @@ public:
 		return NSLOCTEXT("MovieGraph", "ImgSequenceBMPSetting_NodeTitleShort", ".bmp Sequence");
 	}
 
+	virtual FText GetBasicConfigShortDisplayName() const override
+	{
+		return NSLOCTEXT("MovieGraph", "BasicConfigShortName_BMP", "BMP");
+	}
+
 	virtual FText GetKeywords() const override
 	{
 		static const FText Keywords = NSLOCTEXT("MovieGraph", "ImageSequenceOutputNode_BMP_Keywords", "bmp image");
@@ -471,6 +481,11 @@ public:
 		return Keywords;
 	}
 	
+	virtual FText GetBasicConfigShortDisplayName() const override
+	{
+		return NSLOCTEXT("MovieGraph", "BasicConfigShortName_JPG", "JPG");
+	}
+
 	virtual FLinearColor GetNodeTitleColor() const override
 	{
 		return FLinearColor(0.047f, 0.654f, 0.537f);
@@ -558,6 +573,11 @@ public:
 		return Keywords;
 	}
 	
+	virtual FText GetBasicConfigShortDisplayName() const override
+	{
+		return NSLOCTEXT("MovieGraph", "BasicConfigShortName_PNG", "PNG");
+	}
+
 	virtual FLinearColor GetNodeTitleColor() const override
 	{
 		return FLinearColor(0.047f, 0.654f, 0.537f);
