MoviePipeline: Minor style updates in preparation for the new "Basic" UI in the queue.

#Synced-CL 52423999

diff --git a/Engine/Plugins/MovieScene/MovieRenderPipeline/Source/MovieRenderPipelineEditor/Private/Customizations/MoviePipelineNameEditSinks.h b/Engine/Plugins/MovieScene/MovieRenderPipeline/Source/MovieRenderPipelineEditor/Private/Customizations/MoviePipelineNameEditSinks.h
new file mode 100644
--- /dev/null
+++ b/Engine/Plugins/MovieScene/MovieRenderPipeline/Source/MovieRenderPipelineEditor/Private/Customizations/MoviePipelineNameEditSinks.h
@@ -0,0 +1,70 @@
+// Copyright Epic Games, Inc. All Rights Reserved.
+
+#pragma once
+
+#include "IObjectNameEditSink.h"
+#include "MoviePipelineQueue.h"
+
+/** Provides an editable JobName in the details panel name area for UMoviePipelineExecutorJob. */
+class FMoviePipelineJobNameEditSink : public UE::EditorWidgets::IObjectNameEditSink
+{
+public:
+	virtual UClass* GetSupportedClass() const override
+	{
+		return UMoviePipelineExecutorJob::StaticClass();
+	}
+
+	virtual FText GetObjectDisplayName(UObject* Object) const override
+	{
+		const UMoviePipelineExecutorJob* Job = CastChecked<UMoviePipelineExecutorJob>(Object);
+		return FText::FromString(Job->JobName);
+	}
+
+	virtual bool IsObjectDisplayNameReadOnly(UObject* Object) const override
+	{
+		return false;
+	}
+
+	virtual bool SetObjectDisplayName(UObject* Object, FString DisplayName) override
+	{
+		UMoviePipelineExecutorJob* Job = CastChecked<UMoviePipelineExecutorJob>(Object);
+		if (Job->JobName != DisplayName)
+		{
+			Job->Modify();
+			Job->JobName = MoveTemp(DisplayName);
+			return true;
+		}
+		return false;
+	}
+
+	virtual FText GetObjectNameTooltip(UObject* Object) const override
+	{
+		return NSLOCTEXT("MoviePipelineJobNameEditSink", "JobNameTooltip", "The name of this render job");
+	}
+};
+
+/** Provides a read-only OuterName + InnerName display in the details panel name area for UMoviePipelineExecutorShot. */
+class FMoviePipelineShotNameEditSink : public UE::EditorWidgets::IObjectNameEditSink
+{
+public:
+	virtual UClass* GetSupportedClass() const override
+	{
+		return UMoviePipelineExecutorShot::StaticClass();
+	}
+
+	virtual FText GetObjectDisplayName(UObject* Object) const override
+	{
+		const UMoviePipelineExecutorShot* Shot = CastChecked<UMoviePipelineExecutorShot>(Object);
+		return FText::FromString(FString::Printf(TEXT("%s %s"), *Shot->OuterName, *Shot->InnerName));
+	}
+
+	virtual bool IsObjectDisplayNameReadOnly(UObject* Object) const override
+	{
+		return true;
+	}
+
+	virtual FText GetObjectNameTooltip(UObject* Object) const override
+	{
+		return NSLOCTEXT("MoviePipelineJobNameEditSink", "ShotNameTooltip", "The name of this shot");
+	}
+};
diff --git a/Engine/Plugins/MovieScene/MovieRenderPipeline/Source/MovieRenderPipelineEditor/Private/MovieRenderPipelineEditorModule.cpp b/Engine/Plugins/MovieScene/MovieRenderPipeline/Source/MovieRenderPipelineEditor/Private/MovieRenderPipelineEditorModule.cpp
--- a/Engine/Plugins/MovieScene/MovieRenderPipeline/Source/MovieRenderPipelineEditor/Private/MovieRenderPipelineEditorModule.cpp
+++ b/Engine/Plugins/MovieScene/MovieRenderPipeline/Source/MovieRenderPipelineEditor/Private/MovieRenderPipelineEditorModule.cpp
@@ -3,6 +3,7 @@
 #include "MovieRenderPipelineEditorModule.h"
 
 #include "EdGraphUtilities.h"
+#include "EditorWidgetsModule.h"
 #include "MovieRenderPipelineCoreModule.h"
 #include "MoviePipelineExecutor.h"
 #include "WorkspaceMenuStructureModule.h"
@@ -34,9 +35,11 @@
 #include "MovieSceneSection.h"
 #include "Sections/MovieSceneCinematicShotSection.h"
 #include "MoviePipelineConsoleVariableSetting.h"
+#include "ObjectNameEditSinkRegistry.h"
 #include "PropertyEditorModule.h"
 #include "Customizations/ConsoleVariableCustomization.h"
 #include "Customizations/Graph/MovieGraphNamedResolutionCustomization.h"
+#include "Customizations/MoviePipelineNameEditSinks.h"
 #include "LevelSequence.h"
 #include "ToolMenus.h"
 #include "Graph/MovieGraphDataTypes.h"
@@ -50,7 +53,15 @@ FName IMovieRenderPipelineEditorModule::MoviePipelineQueueTabName = "MoviePipeli
 FText IMovieRenderPipelineEditorModule::MoviePipelineQueueTabLabel = LOCTEXT("MovieRenderQueueTab_Label", "Movie Render Queue");
 FName IMovieRenderPipelineEditorModule::MoviePipelineConfigEditorTabName = "MovieRenderPipeline";
 FText IMovieRenderPipelineEditorModule::MoviePipelineConfigEditorTabLabel = LOCTEXT("MovieRenderPipelineTab_Label", "Movie Render Pipeline");
-										 
+
+static TAutoConsoleVariable<int32> CVarMovieRenderPipelineShowNewQueuePanel(
+	TEXT("MovieRenderPipeline.ShowNewQueuePanel"),
+	0,
+	TEXT("When set to 1, the Movie Render Queue tab uses the new queue panel.\n")
+	TEXT("When set to 0 (default), the original queue panel is used.\n")
+	TEXT("Changes take effect the next time the queue tab is opened."),
+	ECVF_Default);
+
 namespace
 {
 	static TSharedRef<SDockTab> SpawnMovieRenderPipelineTab(const FSpawnTabArgs& SpawnTabArgs)
@@ -282,6 +293,16 @@ void FMovieRenderPipelineEditorModule::StartupModule()
 
 	GraphPanelPinFactory = MakeShared<FMovieGraphPanelPinFactory>();
 	FEdGraphUtilities::RegisterVisualPinFactory(GraphPanelPinFactory);
+
+	// Register Object Name sinks for the job/shot so their title appearance is customized in the details panel.
+	if (CVarMovieRenderPipelineShowNewQueuePanel.GetValueOnGameThread() != 0)
+	{
+		const FEditorWidgetsModule& EditorWidgetsModule = FModuleManager::LoadModuleChecked<FEditorWidgetsModule>("EditorWidgets");
+		JobNameEditSink = MakeShared<FMoviePipelineJobNameEditSink>();
+		ShotNameEditSink = MakeShared<FMoviePipelineShotNameEditSink>();
+		EditorWidgetsModule.GetObjectNameEditSinkRegistry()->RegisterObjectNameEditSink(JobNameEditSink.ToSharedRef());
+		EditorWidgetsModule.GetObjectNameEditSinkRegistry()->RegisterObjectNameEditSink(ShotNameEditSink.ToSharedRef());
+	}
 }
 
 void FMovieRenderPipelineEditorModule::ShutdownModule()
@@ -297,6 +318,19 @@ void FMovieRenderPipelineEditorModule::ShutdownModule()
 		FEdGraphUtilities::UnregisterVisualPinFactory(GraphPanelPinFactory);
 		GraphPanelPinFactory.Reset();
 	}
+
+	if (FModuleManager::Get().IsModuleLoaded("EditorWidgets"))
+	{
+		const FEditorWidgetsModule& EditorWidgetsModule = FModuleManager::GetModuleChecked<FEditorWidgetsModule>("EditorWidgets");
+		if (JobNameEditSink.IsValid())
+		{
+			EditorWidgetsModule.GetObjectNameEditSinkRegistry()->UnregisterObjectNameEditSink(JobNameEditSink.ToSharedRef());
+		}
+		if (ShotNameEditSink.IsValid())
+		{
+			EditorWidgetsModule.GetObjectNameEditSinkRegistry()->UnregisterObjectNameEditSink(ShotNameEditSink.ToSharedRef());
+		}
+	}
 }
 
 #undef LOCTEXT_NAMESPACE
diff --git a/Engine/Plugins/MovieScene/MovieRenderPipeline/Source/MovieRenderPipelineEditor/Private/MovieRenderPipelineEditorModule.h b/Engine/Plugins/MovieScene/MovieRenderPipeline/Source/MovieRenderPipelineEditor/Private/MovieRenderPipelineEditorModule.h
--- a/Engine/Plugins/MovieScene/MovieRenderPipeline/Source/MovieRenderPipelineEditor/Private/MovieRenderPipelineEditorModule.h
+++ b/Engine/Plugins/MovieScene/MovieRenderPipeline/Source/MovieRenderPipelineEditor/Private/MovieRenderPipelineEditorModule.h
@@ -7,6 +7,8 @@
 #include "Delegates/IDelegateInstance.h"
 
 class FMovieGraphPanelPinFactory;
+class FMoviePipelineJobNameEditSink;
+class FMoviePipelineShotNameEditSink;
 
 class FMovieRenderPipelineEditorModule : public IMovieRenderPipelineEditorModule
 {
@@ -33,4 +35,10 @@ private:
 
 	/** Pin factory for the render graph. */
 	TSharedPtr<FMovieGraphPanelPinFactory> GraphPanelPinFactory;
+
+	/** ObjectNameEditSink for UMoviePipelineExecutorJob instances in the details panel. */
+	TSharedPtr<FMoviePipelineJobNameEditSink> JobNameEditSink;
+
+	/** ObjectNameEditSink for UMoviePipelineExecutorShot instances in the details panel. */
+	TSharedPtr<FMoviePipelineShotNameEditSink> ShotNameEditSink;
 };
diff --git a/Engine/Plugins/MovieScene/MovieRenderPipeline/Source/MovieRenderPipelineEditor/Private/MovieRenderPipelineStyle.cpp b/Engine/Plugins/MovieScene/MovieRenderPipeline/Source/MovieRenderPipelineEditor/Private/MovieRenderPipelineStyle.cpp
--- a/Engine/Plugins/MovieScene/MovieRenderPipeline/Source/MovieRenderPipelineEditor/Private/MovieRenderPipelineStyle.cpp
+++ b/Engine/Plugins/MovieScene/MovieRenderPipeline/Source/MovieRenderPipelineEditor/Private/MovieRenderPipelineStyle.cpp
@@ -114,6 +114,10 @@ FMovieRenderPipelineStyle::FMovieRenderPipelineStyle()
 
 	Set("MovieRenderPipeline.TabIcon", new FSlateImageBrush(RootToContentDir(TEXT("TabIcon_24x.png")), Icon16x16));
 
+	// Class icons so that SDetailNameArea shows the correct icon for jobs and shots
+	Set("ClassIcon.MoviePipelineExecutorJob", new FSlateBrush(*FAppStyle::Get().GetBrush("ClassIcon.CameraComponent")));
+	Set("ClassIcon.MoviePipelineExecutorShot", new FSlateBrush(*FAppStyle::Get().GetBrush("Sequencer.Tracks.CinematicShot")));
+
 	Set("MovieRenderGraph.CollectionsTree.SmallComboButton", FComboButtonStyle(FAppStyle::Get().GetWidgetStyle<FComboButtonStyle>("SimpleComboButton"))
 		.SetDownArrowPadding(FMargin(0)));
 
diff --git a/Engine/Plugins/MovieScene/MovieRenderPipeline/Source/MovieRenderPipelineEditor/Private/Widgets/SMoviePipelineQueueTabContent.cpp b/Engine/Plugins/MovieScene/MovieRenderPipeline/Source/MovieRenderPipelineEditor/Private/Widgets/SMoviePipelineQueueTabContent.cpp
--- a/Engine/Plugins/MovieScene/MovieRenderPipeline/Source/MovieRenderPipelineEditor/Private/Widgets/SMoviePipelineQueueTabContent.cpp
+++ b/Engine/Plugins/MovieScene/MovieRenderPipeline/Source/MovieRenderPipelineEditor/Private/Widgets/SMoviePipelineQueueTabContent.cpp
@@ -10,17 +10,10 @@
 
 #define LOCTEXT_NAMESPACE "SMoviePipelineQueueTabContent"
 
-static TAutoConsoleVariable<int32> CVarMovieRenderPipelineShowNewQueuePanel(
-	TEXT("MovieRenderPipeline.ShowNewQueuePanel"),
-	0,
-	TEXT("When set to 1, the Movie Render Queue tab uses the new queue panel.\n")
-	TEXT("When set to 0 (default), the original queue panel is used.\n")
-	TEXT("Changes take effect the next time the queue tab is opened."),
-	ECVF_Default);
-
 void SMoviePipelineQueueTabContent::Construct(const FArguments& InArgs)
 {
-	const bool bUseNewPanel = (CVarMovieRenderPipelineShowNewQueuePanel.GetValueOnGameThread() != 0);
+	const IConsoleVariable* ShowNewQueueCvar = IConsoleManager::Get().FindConsoleVariable(TEXT("MovieRenderPipeline.ShowNewQueuePanel"));
+	const bool bUseNewPanel = ShowNewQueueCvar && (ShowNewQueueCvar->GetInt() != 0);
 
 	if (bUseNewPanel)
 	{
