MoviePipeline: Added the datastructure for the new "Basic" config type, and updated its supporting infrastructure to account for it.

#virtualized
#Synced-CL 52423999

diff --git a/Engine/Plugins/MovieScene/MovieRenderPipeline/Source/MovieRenderPipelineCore/Private/Graph/MovieGraphPipeline.cpp b/Engine/Plugins/MovieScene/MovieRenderPipeline/Source/MovieRenderPipelineCore/Private/Graph/MovieGraphPipeline.cpp
--- a/Engine/Plugins/MovieScene/MovieRenderPipeline/Source/MovieRenderPipelineCore/Private/Graph/MovieGraphPipeline.cpp
+++ b/Engine/Plugins/MovieScene/MovieRenderPipeline/Source/MovieRenderPipelineCore/Private/Graph/MovieGraphPipeline.cpp
@@ -2,6 +2,7 @@
 
 #include "Graph/MovieGraphPipeline.h"
 
+#include "MoviePipelineBasicConfig.h"
 #include "MoviePipelineQueue.h"
 #include "MoviePipelineTelemetry.h"
 #include "MoviePipelineUtils.h"
@@ -57,7 +58,7 @@ void UMovieGraphPipeline::Initialize(UMoviePipelineExecutorJob* InJob, const FMo
 		return;
 	}
 
-	if (!ensureAlwaysMsgf(InJob->GetGraphPreset(), TEXT("MoviePipeline cannot be initialized with a null job configuration. Make sure you've created a Graph Config for this job (or use the regular UMovieGraph instead of UMovieGraphPipeline). Aborting.")))
+	if (!ensureAlwaysMsgf(InJob->GetGraphPreset() || InJob->IsUsingBasicConfiguration(), TEXT("MovieGraphPipeline cannot be initialized with a null job configuration. Make sure you've assigned a Graph or Basic config to this job. Aborting.")))
 	{
 		//Shutdown(true);
 		return;
@@ -121,9 +122,10 @@ void UMovieGraphPipeline::Initialize(UMoviePipelineExecutorJob* InJob, const FMo
 
 	const FMovieGraphTraversalContext Context = GetCurrentTraversalContext(false);
 	FString OutError;
-	UMovieGraphEvaluatedConfig* FlattenedConfig = GetCurrentJob()->GetGraphPreset()->CreateFlattenedGraph(Context, OutError);
+	UMovieGraphConfig* GraphPreset = GetCurrentJob()->IsUsingGraphConfiguration() ? GetCurrentJob()->GetGraphPreset() : nullptr;
+	UMovieGraphEvaluatedConfig* FlattenedConfig = GraphPreset ? GraphPreset->CreateFlattenedGraph(Context, OutError) : nullptr;
 
-	// Lay groundwork for Unreal Insights
+	// Lay groundwork for Unreal Insights; this is not supported with Basic configs.
 	if (FlattenedConfig)
 	{
 		constexpr bool bIncludeCDOs = false;
@@ -189,11 +191,51 @@ void UMovieGraphPipeline::DuplicateJobAndConfiguration()
 	JobDuplicationParms.ApplyFlags = RF_Transient;
 	CurrentJobDuplicate = Cast<UMoviePipelineExecutorJob>(StaticDuplicateObjectEx(JobDuplicationParms));
 
-	// The duplicate job is a mix of duplicated objects and non-duplicated objects. Objects that 
+	// Convert Basic config to Graph on the transient duplicate (before the existing duplication logic processes it)
+	if (CurrentJobDuplicate->IsUsingBasicConfiguration())
+	{
+		// Generate job-level graph from the basic config
+		UMovieGraphConfig* GeneratedGraph = UMoviePipelineBasicConfig::GenerateGraph(
+			CurrentJobDuplicate->GetBasicConfig(), CurrentJobDuplicate);
+
+		if (!GeneratedGraph)
+		{
+			UE_LOG(LogMovieRenderPipeline, Error, TEXT("DuplicateJobAndConfiguration: Failed to generate graph from Basic config. The job will have no graph configuration."));
+		}
+
+		// Pre-populate the duplication map so DuplicateConfigRecursive skips these already-transient graphs
+		if (GeneratedGraph)
+		{
+			DuplicatedGraphs.Add(GeneratedGraph, GeneratedGraph);
+		}
+
+		// Generate per-shot graphs (each shot gets a merged config of job + shot overrides)
+		for (UMoviePipelineExecutorShot* Shot : CurrentJobDuplicate->ShotInfo)
+		{
+			if (const UMoviePipelineBasicConfig* ShotOverrides = Shot->GetBasicShotConfig())
+			{
+				UMoviePipelineBasicConfig* MergedConfig = NewObject<UMoviePipelineBasicConfig>(Shot);
+				CurrentJobDuplicate->GetBasicConfig()->MergeOverrides(ShotOverrides, *MergedConfig);
+				UMovieGraphConfig* ShotGraph = UMoviePipelineBasicConfig::GenerateGraph(MergedConfig, Shot);
+				Shot->SetGraphPreset(ShotGraph, false);
+
+				if (ShotGraph)
+				{
+					DuplicatedGraphs.Add(ShotGraph, ShotGraph);
+				}
+			}
+		}
+
+		// Set the job graph after shot graphs are set
+		CurrentJobDuplicate->SetGraphPreset(GeneratedGraph, false);
+	}
+
+	// The duplicate job is a mix of duplicated objects and non-duplicated objects. Objects that
 	// don't have the job as their outer will not have been duplicate (which is good for the World/Sequence),
 	// but this also means that the graph configurations were not duplicated (as they are assets), so we need
 	// to manually duplicate them and update the pointers in the duplicated job.
-	UMovieGraphConfig* DuplicatePrimaryConfig = MoviePipeline::DuplicateConfigRecursive(CurrentJob->GetGraphPreset(), DuplicatedGraphs);
+	// Note: We use CurrentJobDuplicate here (not CurrentJob) so that Basic-config-generated graphs are also processed.
+	UMovieGraphConfig* DuplicatePrimaryConfig = MoviePipeline::DuplicateConfigRecursive(CurrentJobDuplicate->GetGraphPreset(), DuplicatedGraphs);
 	UpdateVariableAssignmentsHelper(CurrentJobDuplicate.Get(), DuplicatedGraphs);
 
 	// Set the graph preset WITHOUT updating assignments. We're doing all updates manually here. SetGraphPreset() will attempt to update the
@@ -202,10 +244,11 @@ void UMovieGraphPipeline::DuplicateJobAndConfiguration()
 	constexpr bool bUpdateVariableAssignments = false;
 	CurrentJobDuplicate->SetGraphPreset(DuplicatePrimaryConfig, bUpdateVariableAssignments);
 
-	for (int32 Index = 0; Index < CurrentJob->ShotInfo.Num(); Index++)
+	for (int32 Index = 0; Index < CurrentJobDuplicate->ShotInfo.Num(); Index++)
 	{
 		// Now for each shot we need to duplicate its config (if any)
-		if (UMovieGraphConfig* ShotGraphPreset = CurrentJob->ShotInfo[Index]->GetGraphPreset())
+		// Note: We use CurrentJobDuplicate here (not CurrentJob) so that Basic-config-generated graphs are also processed.
+		if (UMovieGraphConfig* ShotGraphPreset = CurrentJobDuplicate->ShotInfo[Index]->GetGraphPreset())
 		{
 			UMovieGraphConfig* DuplicateShotConfig = MoviePipeline::DuplicateConfigRecursive(ShotGraphPreset, DuplicatedGraphs);
 			CurrentJobDuplicate->ShotInfo[Index]->SetGraphPreset(DuplicateShotConfig, bUpdateVariableAssignments);
diff --git a/Engine/Plugins/MovieScene/MovieRenderPipeline/Source/MovieRenderPipelineCore/Private/Graph/Nodes/MovieGraphGlobalGameOverrides.cpp b/Engine/Plugins/MovieScene/MovieRenderPipeline/Source/MovieRenderPipelineCore/Private/Graph/Nodes/MovieGraphGlobalGameOverrides.cpp
--- a/Engine/Plugins/MovieScene/MovieRenderPipeline/Source/MovieRenderPipelineCore/Private/Graph/Nodes/MovieGraphGlobalGameOverrides.cpp
+++ b/Engine/Plugins/MovieScene/MovieRenderPipeline/Source/MovieRenderPipelineCore/Private/Graph/Nodes/MovieGraphGlobalGameOverrides.cpp
@@ -274,6 +274,11 @@ TSubclassOf<AGameModeBase> UMovieGraphGlobalGameOverridesNode::GetGameModeOverri
 			return GameOverridesNode ? GameOverridesNode->SoftGameModeOverride.LoadSynchronous() : nullptr;
 		}
 	}
+	else if (InJob->IsUsingBasicConfiguration())
+	{
+		// Basic configuration does not support game mode overrides.
+		return nullptr;
+	}
 	else
 	{
 		TArray<UMoviePipelineSetting*> AllSettings = InJob->GetConfiguration()->GetAllSettings();
@@ -281,7 +286,7 @@ TSubclassOf<AGameModeBase> UMovieGraphGlobalGameOverridesNode::GetGameModeOverri
 		{
 			return InSetting->GetClass() == UMoviePipelineGameOverrideSetting::StaticClass();
 		});
-		
+
 		if (GameOverridesPtr)
 		{
 			if (UMoviePipelineSetting* Setting = *GameOverridesPtr)
diff --git a/Engine/Plugins/MovieScene/MovieRenderPipeline/Source/MovieRenderPipelineCore/Private/MoviePipelineBasicConfig.cpp b/Engine/Plugins/MovieScene/MovieRenderPipeline/Source/MovieRenderPipelineCore/Private/MoviePipelineBasicConfig.cpp
new file mode 100644
--- /dev/null
+++ b/Engine/Plugins/MovieScene/MovieRenderPipeline/Source/MovieRenderPipelineCore/Private/MoviePipelineBasicConfig.cpp
@@ -0,0 +1,506 @@
+// Copyright Epic Games, Inc. All Rights Reserved.
+
+#include "MoviePipelineBasicConfig.h"
+
+#include "Graph/MovieGraphConfig.h"
+#include "Graph/MovieGraphNode.h"
+#include "Graph/Nodes/MovieGraphFileOutputNode.h"
+#include "Graph/Nodes/MovieGraphGlobalOutputSettingNode.h"
+#include "Graph/Nodes/MovieGraphSamplingMethodNode.h"
+#include "Graph/Nodes/MovieGraphWarmUpSettingNode.h"
+#include "MovieRenderPipelineCoreModule.h"
+#include "UObject/Package.h"
+#include "UObject/UObjectGlobals.h"
+
+#if WITH_EDITOR
+#include "PackageHelperFunctions.h"
+#include "Settings/EditorLoadingSavingSettings.h"
+#endif
+
+#include UE_INLINE_GENERATED_CPP_BY_NAME(MoviePipelineBasicConfig)
+
+namespace UE::MoviePipeline::BasicConfig
+{
+	// Reflected property names on nodes from the MovieRenderPipelineRenderPasses module (avoiding compile-time dependency)
+	static const FName PropName_SpatialSampleCount					= TEXT("SpatialSampleCount");
+	static const FName PropName_bOverride_SpatialSampleCount		= TEXT("bOverride_SpatialSampleCount");
+	static const FName PropName_bEnableDenoiser						= TEXT("bEnableDenoiser");
+	static const FName PropName_bOverride_bEnableDenoiser			= TEXT("bOverride_bEnableDenoiser");
+	static const FName PropName_DenoiserType						= TEXT("DenoiserType");
+	static const FName PropName_bOverride_DenoiserType				= TEXT("bOverride_DenoiserType");
+	static const FName PropName_bEnableBurnIn						= TEXT("bEnableBurnIn");
+	static const FName PropName_bOverride_bEnableBurnIn				= TEXT("bOverride_bEnableBurnIn");
+	static const FName PropName_BurnInClass							= TEXT("BurnInClass");
+	static const FName PropName_bOverride_BurnInClass				= TEXT("bOverride_BurnInClass");
+	static const FName PropName_bCompositeOntoFinalImage			= TEXT("bCompositeOntoFinalImage");
+	static const FName PropName_bOverride_bCompositeOntoFinalImage	= TEXT("bOverride_bCompositeOntoFinalImage");
+
+	/** Set a property on a UObject by name using its serialized string representation. */
+	static void SetReflectedProperty(UObject* InObject, const FName PropertyName, const FString& Value)
+	{
+		if (!InObject)
+		{
+			return;
+		}
+
+		const FProperty* Property = FindFProperty<FProperty>(InObject->GetClass(), PropertyName);
+		if (!Property)
+		{
+			UE_LOG(LogMovieRenderPipeline, Warning, TEXT("UMoviePipelineBasicConfig: Could not find property '%s' on '%s'"), *PropertyName.ToString(), *InObject->GetClass()->GetName());
+			return;
+		}
+
+		void* ValuePtr = Property->ContainerPtrToValuePtr<void>(InObject);
+		if (Property->ImportText_Direct(*Value, ValuePtr, InObject, PPF_None) == nullptr)
+		{
+			UE_LOG(LogMovieRenderPipeline, Warning, TEXT("UMoviePipelineBasicConfig: Failed to set property '%s' on '%s' from value '%s'"),
+				*PropertyName.ToString(), *InObject->GetClass()->GetName(), *Value);
+		}
+	}
+
+	/** Set a bOverride_ flag on a UObject to true. */
+	static void SetOverrideFlag(UObject* InObject, const FName FlagPropertyName)
+	{
+		if (!InObject)
+		{
+			return;
+		}
+
+		if (const FBoolProperty* BoolProp = CastField<FBoolProperty>(FindFProperty<FProperty>(InObject->GetClass(), FlagPropertyName)))
+		{
+			void* ValuePtr = BoolProp->ContainerPtrToValuePtr<void>(InObject);
+			BoolProp->SetPropertyValue(ValuePtr, true);
+		}
+	}
+
+	/** Template graph asset path. */
+	static const FSoftObjectPath TemplateGraphPath(TEXT("/MovieRenderPipeline/BasicConfigTemplateGraph.BasicConfigTemplateGraph"));
+	
+	/** The path of the user-saved default for Basic configurations. This asset may not exist if a default has not been saved. */
+	static const FString DefaultBasicConfigPath = TEXT("/Temp/MovieRenderPipeline/BasicConfigDefaults");
+
+#if WITH_EDITOR
+	/** Lay out all nodes in each branch left-to-right with a fixed stride. */
+	static void LayoutGraphBranches(const UMovieGraphConfig* Graph)
+	{
+		constexpr int32 NodeWidth   	= 200;
+		constexpr int32 NodeSpacing 	= 100;
+		constexpr int32 Stride      	= NodeWidth + NodeSpacing;
+		constexpr int32 BranchYSpacing	= 300;
+
+		// Globals branch first, then all others
+		TArray<FName> BranchOrder = Graph->GetBranchNames();
+		BranchOrder.Remove(UMovieGraphNode::GlobalsPinName);
+		BranchOrder.Insert(UMovieGraphNode::GlobalsPinName, 0);
+
+		const int32 InputNodeX = Graph->GetInputNode()->GetNodePosX();
+		int32 MaxNodeX = InputNodeX;
+		int32 BranchY  = 0;
+
+		for (const FName& BranchName : BranchOrder)
+		{
+			// Walk from InputNode's branch output pin through intermediate nodes to the OutputNode
+			TArray<UMovieGraphNode*> IntermediateNodes;
+			if (const UMovieGraphPin* BranchOutPin = Graph->GetInputNode()->GetOutputPin(BranchName))
+			{
+				const UMovieGraphPin* Next = BranchOutPin->GetFirstConnectedPin();
+				while (Next && Next->Node != Graph->GetOutputNode())
+				{
+					IntermediateNodes.Add(Next->Node);
+					const UMovieGraphPin* NodeOut = Next->Node->GetOutputPin(FName());
+					if (!NodeOut)
+					{
+						break;
+					}
+					
+					Next = NodeOut->GetFirstConnectedPin();
+				}
+			}
+
+			// Assign X positions left-to-right, starting one Stride past InputNode
+			for (int32 i = 0; i < IntermediateNodes.Num(); ++i)
+			{
+				const int32 NodeX = InputNodeX + Stride * (i + 1);
+				IntermediateNodes[i]->SetNodePosX(NodeX);
+				IntermediateNodes[i]->SetNodePosY(BranchY);
+				MaxNodeX = FMath::Max(MaxNodeX, NodeX);
+			}
+
+			BranchY += BranchYSpacing;
+		}
+
+		// Center InputNode and OutputNode vertically across all branches
+		const int32 TotalBranches = BranchOrder.Num();
+		const int32 MidY = (TotalBranches > 1) ? ((TotalBranches - 1) * BranchYSpacing) / 2 : 0;
+		Graph->GetInputNode()->SetNodePosY(MidY);
+
+		// Push OutputNode to the right of all intermediate nodes
+		Graph->GetOutputNode()->SetNodePosX(MaxNodeX + Stride);
+		Graph->GetOutputNode()->SetNodePosY(MidY);
+	}
+#endif // WITH_EDITOR
+}
+
+UMoviePipelineBasicConfig::UMoviePipelineBasicConfig()
+	: bOverride_OutputDirectory(0)
+	, bOverride_FileNameFormat(0)
+	, bOverride_OutputResolution(0)
+	, bOverride_EnabledOutputTypes(0)
+	, bOverride_bEnableCustomStartFrame(0)
+	, bOverride_CustomStartFrame(0)
+	, bOverride_bEnableCustomEndFrame(0)
+	, bOverride_CustomEndFrame(0)
+	, bOverride_bUseDeferredRenderer(0)
+	, bOverride_DeferredSpatialSampleCount(0)
+	, bOverride_bUsePathTracedRenderer(0)
+	, bOverride_PathTracedSpatialSampleCount(0)
+	, bOverride_bPathTracedEnableDenoiser(0)
+	, bOverride_PathTracedDenoiserType(0)
+	, bOverride_NumWarmUpFrames(0)
+	, bOverride_TemporalSampleCount(0)
+	, bOverride_BurnInClass(0)
+	, bEnableCustomStartFrame(false)
+	, CustomStartFrame(0)
+	, bEnableCustomEndFrame(false)
+	, CustomEndFrame(0)
+	, bUseDeferredRenderer(true)
+	, DeferredSpatialSampleCount(1)
+	, bUsePathTracedRenderer(false)
+	, PathTracedSpatialSampleCount(1)
+	, bPathTracedEnableDenoiser(false)
+	, PathTracedDenoiserType(EMoviePipelineBasicDenoiserType::Spatial)
+	, NumWarmUpFrames(64)
+	, TemporalSampleCount(1)
+{
+	OutputDirectory.Path = TEXT("{project_dir}/Saved/MovieRenders/");
+	FileNameFormat = TEXT("{sequence_name}.{frame_number}");
+	OutputResolution = FMovieGraphNamedResolution(FMovieGraphNamedResolution::DefaultResolutionName, FMovieGraphNamedResolution::DefaultResolution, FString());
+
+	// Default to PNG output
+	EnabledOutputTypes.Add(TSoftClassPtr<UMovieGraphFileOutputNode>(
+		FSoftClassPath(TEXT("/Script/MovieRenderPipelineRenderPasses.MovieGraphImageSequenceOutputNode_PNG"))));
+
+	BurnInClass = FSoftClassPath(TEXT("/MovieRenderPipeline/Blueprints/Graph/DefaultGraphBurnIn.DefaultGraphBurnIn_C"));
+}
+
+bool UMoviePipelineBasicConfig::HasAnyOverrides() const
+{
+	for (TFieldIterator<FBoolProperty> It(GetClass()); It; ++It)
+	{
+		if (It->GetName().StartsWith(TEXT("bOverride_")) && It->GetPropertyValue_InContainer(this))
+		{
+			return true;
+		}
+	}
+	return false;
+}
+
+void UMoviePipelineBasicConfig::EnforceInvariants()
+{
+	if (!bUseDeferredRenderer && !bUsePathTracedRenderer)
+	{
+		bUseDeferredRenderer = true;
+	}
+
+	if (EnabledOutputTypes.IsEmpty())
+	{
+		EnabledOutputTypes.Add(TSoftClassPtr<UMovieGraphFileOutputNode>(
+			FSoftClassPath(TEXT("/Script/MovieRenderPipelineRenderPasses.MovieGraphImageSequenceOutputNode_PNG"))));
+	}
+}
+
+#if WITH_EDITOR
+void UMoviePipelineBasicConfig::PostEditChangeProperty(FPropertyChangedEvent& PropertyChangedEvent)
+{
+	Super::PostEditChangeProperty(PropertyChangedEvent);
+	EnforceInvariants();
+}
+#endif
+
+void UMoviePipelineBasicConfig::MergeOverrides(const UMoviePipelineBasicConfig* InShotOverrides, UMoviePipelineBasicConfig& OutMerged) const
+{
+	// For each property, use the shot override value if its bOverride_ flag is set; otherwise use the job value.
+	// When InShotOverrides is null, all flags are effectively false, so job values are always used.
+#define MERGE_PROPERTY(PropName) \
+	OutMerged.bOverride_##PropName = (InShotOverrides && InShotOverrides->bOverride_##PropName) ? InShotOverrides->bOverride_##PropName : bOverride_##PropName; \
+	OutMerged.PropName = (InShotOverrides && InShotOverrides->bOverride_##PropName) ? InShotOverrides->PropName : PropName;
+
+	MERGE_PROPERTY(OutputDirectory);
+	MERGE_PROPERTY(FileNameFormat);
+	MERGE_PROPERTY(OutputResolution);
+	MERGE_PROPERTY(EnabledOutputTypes);
+	MERGE_PROPERTY(bEnableCustomStartFrame);
+	MERGE_PROPERTY(CustomStartFrame);
+	MERGE_PROPERTY(bEnableCustomEndFrame);
+	MERGE_PROPERTY(CustomEndFrame);
+	MERGE_PROPERTY(bUseDeferredRenderer);
+	MERGE_PROPERTY(DeferredSpatialSampleCount);
+	MERGE_PROPERTY(bUsePathTracedRenderer);
+	MERGE_PROPERTY(PathTracedSpatialSampleCount);
+	MERGE_PROPERTY(bPathTracedEnableDenoiser);
+	MERGE_PROPERTY(PathTracedDenoiserType);
+	MERGE_PROPERTY(NumWarmUpFrames);
+	MERGE_PROPERTY(TemporalSampleCount);
+	MERGE_PROPERTY(BurnInClass);
+
+#undef MERGE_PROPERTY
+}
+
+UMovieGraphConfig* UMoviePipelineBasicConfig::GenerateGraph(UMoviePipelineBasicConfig* InConfig, UObject* InOuter)
+{
+	using namespace UE::MoviePipeline::BasicConfig;
+
+	if (!InConfig)
+	{
+		UE_LOG(LogMovieRenderPipeline, Error, TEXT("UMoviePipelineBasicConfig::GenerateGraph called with null config."));
+		return nullptr;
+	}
+
+	// Enforce invariants defensively — config may have been mutated programmatically or via scripting
+	// without going through PostEditChangeProperty.
+	InConfig->EnforceInvariants();
+
+	// Load and duplicate the default graph so we can modify it freely.
+	const UMovieGraphConfig* TemplateGraph = Cast<UMovieGraphConfig>(TemplateGraphPath.TryLoad());
+	if (!TemplateGraph)
+	{
+		UE_LOG(LogMovieRenderPipeline, Error, TEXT("UMoviePipelineBasicConfig::GenerateGraph: Failed to load the template graph."));
+		return nullptr;
+	}
+
+	UMovieGraphConfig* Graph = DuplicateObject<UMovieGraphConfig>(TemplateGraph, InOuter);
+	if (!Graph)
+	{
+		UE_LOG(LogMovieRenderPipeline, Error, TEXT("UMoviePipelineBasicConfig::GenerateGraph: Failed to duplicate the template graph."));
+		return nullptr;
+	}
+
+	// Find an existing node of the given class in the graph. Returns the first match, or nullptr.
+	auto FindExistingNode = [Graph](const UClass* NodeClass) -> UMovieGraphNode*
+	{
+		for (const TObjectPtr<UMovieGraphNode>& Node : Graph->GetNodes())
+		{
+			if (Node && Node->GetClass() == NodeClass)
+			{
+				return Node;
+			}
+		}
+		return nullptr;
+	};
+
+	// Find an existing node of the given class in the duplicated graph, or insert a new one after the
+	// InputNode on the specified pin. New nodes are inserted between the InputNode and whatever is
+	// already downstream, so existing nodes (e.g. Render Layer) stay near the OutputNode.
+	UMovieGraphNode* InputNode = Graph->GetInputNode();
+	auto FindOrInsertNode = [&](const FName& InputPinName, const TSubclassOf<UMovieGraphNode> NodeClass) -> UMovieGraphNode*
+	{
+		if (!NodeClass) { return nullptr; }
+
+		if (UMovieGraphNode* Existing = FindExistingNode(NodeClass))
+		{
+			return Existing;
+		}
+
+		return Graph->InsertAfter(InputNode, NodeClass, InputPinName);
+	};
+
+	// Gather pin names from InputNode's output pins. The Globals pin and all other pins (render layers).
+	const FName GlobalsPin = UMovieGraphNode::GlobalsPinName;
+	FName RenderLayerPin;
+	for (const UMovieGraphPin* Pin : InputNode->GetOutputPins())
+	{
+		if (Pin && Pin->Properties.Label != GlobalsPin)
+		{
+			RenderLayerPin = Pin->Properties.Label;
+			break;
+		}
+	}
+
+	// ── Globals branch ────────────────────────────────────────────────────────
+
+	// GlobalOutputSettingNode: output directory, resolution, optional custom frame range
+	UMovieGraphGlobalOutputSettingNode* GlobalOutputNode = Cast<UMovieGraphGlobalOutputSettingNode>(
+		FindOrInsertNode(GlobalsPin, UMovieGraphGlobalOutputSettingNode::StaticClass()));
+	if (GlobalOutputNode)
+	{
+		GlobalOutputNode->bOverride_OutputDirectory = true;
+		GlobalOutputNode->OutputDirectory = InConfig->OutputDirectory;
+		GlobalOutputNode->bOverride_OutputResolution = true;
+		GlobalOutputNode->OutputResolution = InConfig->OutputResolution;
+		if (InConfig->bEnableCustomStartFrame)
+		{
+			GlobalOutputNode->bOverride_CustomPlaybackRangeStart = true;
+			GlobalOutputNode->CustomPlaybackRangeStart.Type = EMovieGraphSequenceRangeType::Custom;
+			GlobalOutputNode->CustomPlaybackRangeStart.Value = InConfig->CustomStartFrame;
+		}
+		if (InConfig->bEnableCustomEndFrame)
+		{
+			GlobalOutputNode->bOverride_CustomPlaybackRangeEnd = true;
+			GlobalOutputNode->CustomPlaybackRangeEnd.Type = EMovieGraphSequenceRangeType::Custom;
+			GlobalOutputNode->CustomPlaybackRangeEnd.Value = InConfig->CustomEndFrame;
+		}
+	}
+
+	// SamplingMethodNode: temporal sample count
+	if (UMovieGraphSamplingMethodNode* SamplingNode = Cast<UMovieGraphSamplingMethodNode>(
+		FindOrInsertNode(GlobalsPin, UMovieGraphSamplingMethodNode::StaticClass())))
+	{
+		SamplingNode->bOverride_TemporalSampleCount = true;
+		SamplingNode->TemporalSampleCount = InConfig->TemporalSampleCount;
+	}
+
+	// WarmUpSettingNode: warm-up frame count
+	if (UMovieGraphWarmUpSettingNode* WarmUpNode = Cast<UMovieGraphWarmUpSettingNode>(
+		FindOrInsertNode(GlobalsPin, UMovieGraphWarmUpSettingNode::StaticClass())))
+	{
+		WarmUpNode->bOverride_NumWarmUpFrames = true;
+		WarmUpNode->NumWarmUpFrames = InConfig->NumWarmUpFrames;
+	}
+
+	// ── Render layer branch ───────────────────────────────────────────────────
+	if (!RenderLayerPin.IsNone())
+	{
+		// Deferred renderer — only add/configure if enabled
+		if (InConfig->bUseDeferredRenderer)
+		{
+			const TSubclassOf<UMovieGraphNode> DeferredRendererClass =
+				FSoftClassPath(TEXT("/Script/MovieRenderPipelineRenderPasses.MovieGraphDeferredRenderPassNode")).TryLoadClass<UMovieGraphNode>();
+
+			if (UMovieGraphNode* DeferredNode = FindOrInsertNode(RenderLayerPin, DeferredRendererClass))
+			{
+				SetOverrideFlag(DeferredNode, PropName_bOverride_SpatialSampleCount);
+				SetReflectedProperty(DeferredNode, PropName_SpatialSampleCount, FString::FromInt(InConfig->DeferredSpatialSampleCount));
+			}
+		}
+
+		// PathTracer renderer — only add/configure if enabled
+		if (InConfig->bUsePathTracedRenderer)
+		{
+			const TSubclassOf<UMovieGraphNode> PathTracerRendererClass =
+				FSoftClassPath(TEXT("/Script/MovieRenderPipelineRenderPasses.MovieGraphPathTracerRenderPassNode")).TryLoadClass<UMovieGraphNode>();
+
+			if (UMovieGraphNode* PathTracerNode = FindOrInsertNode(RenderLayerPin, PathTracerRendererClass))
+			{
+				SetOverrideFlag(PathTracerNode, PropName_bOverride_SpatialSampleCount);
+				SetReflectedProperty(PathTracerNode, PropName_SpatialSampleCount, FString::FromInt(InConfig->PathTracedSpatialSampleCount));
+				SetOverrideFlag(PathTracerNode, PropName_bOverride_bEnableDenoiser);
+				SetReflectedProperty(PathTracerNode, PropName_bEnableDenoiser, InConfig->bPathTracedEnableDenoiser ? TEXT("true") : TEXT("false"));
+				SetOverrideFlag(PathTracerNode, PropName_bOverride_DenoiserType);
+				SetReflectedProperty(PathTracerNode, PropName_DenoiserType,
+					StaticEnum<EMoviePipelineBasicDenoiserType>()->GetNameStringByValue(static_cast<int64>(InConfig->PathTracedDenoiserType)));
+			}
+		}
+
+		// File output nodes from EnabledOutputTypes — route to the correct branch based on restriction.
+		// Also configures burn-in on each output node that supports it (when the burn-in override is active).
+		for (const TSoftClassPtr<UMovieGraphFileOutputNode>& OutputClassPtr : InConfig->EnabledOutputTypes)
+		{
+			UClass* OutputClass = OutputClassPtr.LoadSynchronous();
+			if (!OutputClass || OutputClass->HasAnyClassFlags(CLASS_Abstract))
+			{
+				UE_LOG(LogMovieRenderPipeline, Warning,
+					TEXT("UMoviePipelineBasicConfig::GenerateGraph: Skipping invalid/abstract output type '%s'."),
+					*OutputClassPtr.ToString());
+				continue;
+			}
+
+			// Check branch restriction via CDO — Globals-restricted nodes go on the Globals branch
+			const UMovieGraphFileOutputNode* CDO = GetDefault<UMovieGraphFileOutputNode>(OutputClass);
+			const bool bGlobalsOnly = CDO && CDO->GetBranchRestriction() == EMovieGraphBranchRestriction::Globals;
+			const FName& TargetPin = bGlobalsOnly ? GlobalsPin : RenderLayerPin;
+
+			UMovieGraphFileOutputNode* OutputNode = Cast<UMovieGraphFileOutputNode>(
+				FindOrInsertNode(TargetPin, OutputClass));
+			if (OutputNode)
+			{
+				OutputNode->bOverride_FileNameFormat = true;
+				OutputNode->FileNameFormat = InConfig->FileNameFormat;
+
+				// Burn-in: if the override is active and the node supports it, enable burn-in
+				if (InConfig->bOverride_BurnInClass && InConfig->BurnInClass.IsValid()
+					&& FindFProperty<FProperty>(OutputNode->GetClass(), PropName_bEnableBurnIn))
+				{
+					SetOverrideFlag(OutputNode, PropName_bOverride_bEnableBurnIn);
+					SetReflectedProperty(OutputNode, PropName_bEnableBurnIn, TEXT("true"));
+					SetOverrideFlag(OutputNode, PropName_bOverride_BurnInClass);
+					SetReflectedProperty(OutputNode, PropName_BurnInClass, InConfig->BurnInClass.ToString());
+
+					// Enable compositing only if the node exposes bCompositeOntoFinalImage
+					if (FindFProperty<FProperty>(OutputNode->GetClass(), PropName_bCompositeOntoFinalImage))
+					{
+						SetOverrideFlag(OutputNode, PropName_bOverride_bCompositeOntoFinalImage);
+						SetReflectedProperty(OutputNode, PropName_bCompositeOntoFinalImage, TEXT("true"));
+					}
+				}
+			}
+		}
+	}
+
+	// ── Layout ────────────────────────────────────────────────────────────────
+#if WITH_EDITOR
+	LayoutGraphBranches(Graph);
+#endif
+
+	return Graph;
+}
+
+const UMoviePipelineBasicConfig* UMoviePipelineBasicConfig::GetSavedDefault()
+{
+	using namespace UE::MoviePipeline::BasicConfig;
+
+	// Use FindPackage first to avoid triggering the async loading path if already in memory.
+	const UPackage* Package = FindPackage(nullptr, *DefaultBasicConfigPath);
+	if (!Package)
+	{
+		Package = LoadPackage(nullptr, *DefaultBasicConfigPath, LOAD_None);
+	}
+
+	if (Package)
+	{
+		return Cast<UMoviePipelineBasicConfig>(FindObjectWithOuter(Package, StaticClass()));
+	}
+
+	return nullptr;
+}
+
+#if WITH_EDITOR
+void UMoviePipelineBasicConfig::SaveAsDefault(const UMoviePipelineBasicConfig* InConfig)
+{
+	using namespace UE::MoviePipeline::BasicConfig;
+
+	if (!InConfig)
+	{
+		return;
+	}
+
+	UPackage* Package = CreatePackage(*DefaultBasicConfigPath);
+	UMoviePipelineBasicConfig* DuplicatedConfig = CastChecked<UMoviePipelineBasicConfig>(
+		StaticDuplicateObject(InConfig, Package, TEXT("BasicConfigDefaults")));
+	DuplicatedConfig->SetFlags(RF_Public | RF_Transactional | RF_Standalone);
+
+	// Zero out all bOverride_ flags via reflection — these are only meaningful for shot overrides, not for default values.
+	for (TFieldIterator<FBoolProperty> It(StaticClass()); It; ++It)
+	{
+		if (It->GetName().StartsWith(TEXT("bOverride_")))
+		{
+			It->SetPropertyValue_InContainer(DuplicatedConfig, false);
+		}
+	}
+
+	// Save to disk. Temporarily disable SCC auto-add for generated local files.
+	const FString PackageFileName = FPackageName::LongPackageNameToFilename(DefaultBasicConfigPath, FPackageName::GetAssetPackageExtension());
+	UEditorLoadingSavingSettings* SaveSettings = GetMutableDefault<UEditorLoadingSavingSettings>();
+	const uint32 bSCCAutoAddNewFiles = SaveSettings->bSCCAutoAddNewFiles;
+	SaveSettings->bSCCAutoAddNewFiles = 0;
+
+	const bool bSuccess = SavePackageHelper(Package, *PackageFileName);
+
+	SaveSettings->bSCCAutoAddNewFiles = bSCCAutoAddNewFiles;
+
+	if (!bSuccess)
+	{
+		UE_LOG(LogMovieRenderPipeline, Warning,
+			TEXT("UMoviePipelineBasicConfig::SaveAsDefault: Failed to save default config to '%s'."), *PackageFileName);
+	}
+}
+#endif // WITH_EDITOR
diff --git a/Engine/Plugins/MovieScene/MovieRenderPipeline/Source/MovieRenderPipelineCore/Private/MoviePipelineInProcessExecutor.cpp b/Engine/Plugins/MovieScene/MovieRenderPipeline/Source/MovieRenderPipelineCore/Private/MoviePipelineInProcessExecutor.cpp
--- a/Engine/Plugins/MovieScene/MovieRenderPipeline/Source/MovieRenderPipelineCore/Private/MoviePipelineInProcessExecutor.cpp
+++ b/Engine/Plugins/MovieScene/MovieRenderPipeline/Source/MovieRenderPipelineCore/Private/MoviePipelineInProcessExecutor.cpp
@@ -117,8 +117,9 @@ void UMoviePipelineInProcessExecutor::OnMapLoadFinished(UWorld* NewWorld)
 		MoviePipelineClass = UMoviePipeline::StaticClass();
 	}
 
-	// Force the pipeline to the default graph pipeline if we're using a graph and the pipeline isn't graph-compatible
-	if (CurrentJob->IsUsingGraphConfiguration() && !MoviePipelineClass->IsChildOf(UMovieGraphPipeline::StaticClass()))
+	// Force the pipeline to the default graph pipeline if we're using a graph or basic configuration and the pipeline isn't graph-compatible
+	if ((CurrentJob->IsUsingGraphConfiguration() || CurrentJob->IsUsingBasicConfiguration())
+		&& !MoviePipelineClass->IsChildOf(UMovieGraphPipeline::StaticClass()))
 	{
 		MoviePipelineClass = UMovieGraphPipeline::StaticClass();
 	}
@@ -146,7 +147,7 @@ void UMoviePipelineInProcessExecutor::OnMapLoadFinished(UWorld* NewWorld)
 	// shut down before trying to access the (now non-existent) player.
 	FCoreDelegates::GetApplicationWillTerminateDelegate().AddUObject(this, &UMoviePipelineInProcessExecutor::OnApplicationQuit);
 
-	// Wait until we actually recieved the right map and created the pipeline before saying that we're actively rendering
+	// Wait until we actually received the right map and created the pipeline before saying that we're actively rendering
 	bIsRendering = true;
 	
 	if (ExecutorSettings->InitialDelayFrameCount == 0)
diff --git a/Engine/Plugins/MovieScene/MovieRenderPipeline/Source/MovieRenderPipelineCore/Private/MoviePipelineQueue.cpp b/Engine/Plugins/MovieScene/MovieRenderPipeline/Source/MovieRenderPipelineCore/Private/MoviePipelineQueue.cpp
--- a/Engine/Plugins/MovieScene/MovieRenderPipeline/Source/MovieRenderPipelineCore/Private/MoviePipelineQueue.cpp
+++ b/Engine/Plugins/MovieScene/MovieRenderPipeline/Source/MovieRenderPipelineCore/Private/MoviePipelineQueue.cpp
@@ -7,6 +7,7 @@
 #include "MoviePipelineBlueprintLibrary.h"
 #include "MoviePipelineUtils.h"
 #include "MovieRenderPipelineCoreModule.h"
+#include "MovieRenderPipelineCoreObjectVersion.h"
 
 #include "UObject/UObjectGlobals.h"
 
@@ -198,6 +199,30 @@ void UMoviePipelineExecutorJob::PostEditChangeProperty(FPropertyChangedEvent& Pr
 }
 #endif
 
+void UMoviePipelineExecutorJob::Serialize(FArchive& Ar)
+{
+	Super::Serialize(Ar);
+
+	Ar.UsingCustomVersion(FMovieRenderPipelineCoreObjectVersion::GUID);
+
+	if (Ar.IsLoading() && Ar.CustomVer(FMovieRenderPipelineCoreObjectVersion::GUID) < FMovieRenderPipelineCoreObjectVersion::JobConfigModeAdded)
+	{
+		// Derive ConfigMode from pointer state for assets saved before the enum was introduced.
+		if (BasicConfig != nullptr)
+		{
+			ConfigMode = EMoviePipelineConfigMode::Basic;
+		}
+		else if (!GraphPreset.IsNull())
+		{
+			ConfigMode = EMoviePipelineConfigMode::Graph;
+		}
+		else
+		{
+			ConfigMode = EMoviePipelineConfigMode::Preset;
+		}
+	}
+}
+
 void UMoviePipelineExecutorJob::PreSave(FObjectPreSaveContext ObjectSaveContext)
 {
 	Super::PreSave(ObjectSaveContext);
@@ -215,6 +240,16 @@ void UMoviePipelineExecutorJob::PostLoad()
 	FCoreUObjectDelegates::OnObjectPreSave.AddUObject(this, &UMoviePipelineExecutorJob::OnGraphPreSave);
 #endif
 
+	// Validate that Graph mode actually has a resolvable graph. The Serialize migration sets ConfigMode
+	// based on !GraphPreset.IsNull() (non-empty path), but the referenced asset may have been deleted.
+	if (ConfigMode == EMoviePipelineConfigMode::Graph && GetGraphPreset() == nullptr)
+	{
+		UE_LOG(LogMovieRenderPipeline, Warning,
+			TEXT("Job '%s' was in Graph mode but the graph preset could not be loaded. Falling back to Preset mode."),
+			*GetName());
+		ConfigMode = EMoviePipelineConfigMode::Preset;
+	}
+
 	// Ideally the variable assignments would be updated here, but currently GraphPreset.LoadSynchronous() does not provide a fully-loaded
 	// object (GraphPreset will only be partially loaded at this point and its variables will be empty).
 }
@@ -257,13 +292,19 @@ void UMoviePipelineExecutorJob::SetGraphPreset(const UMovieGraphConfig* InGraphP
 #if WITH_EDITOR
 	Modify();
 #endif
-	
+
 	GraphPreset = const_cast<UMovieGraphConfig*>(InGraphPreset);
 
-	// If the graph is being cleared out, also clear out all graphs on the job's shots. A shot cannot have a graph assigned to it while
-	// the parent job is using a legacy config.
-	if (InGraphPreset == nullptr)
+	if (InGraphPreset != nullptr)
+	{
+		ConfigMode = EMoviePipelineConfigMode::Graph;
+	}
+	else if (ConfigMode == EMoviePipelineConfigMode::Graph)
 	{
+		// Clearing the graph from Graph mode switches back to Preset mode and cascades to all shots. This is mostly done
+		// for backwards-compatibility purposes.
+		ConfigMode = EMoviePipelineConfigMode::Preset;
+
 		for (const TObjectPtr<UMoviePipelineExecutorShot>& Shot : ShotInfo)
 		{
 			Shot->SetGraphPreset(nullptr);
@@ -325,6 +366,7 @@ void UMoviePipelineExecutorJob::SetConfiguration(UMoviePipelinePrimaryConfig* In
 	{
 		Configuration->CopyFrom(InPreset);
 		PresetOrigin = nullptr;
+		ConfigMode = EMoviePipelineConfigMode::Preset;
 	}
 }
 
@@ -334,6 +376,7 @@ void UMoviePipelineExecutorJob::SetPresetOrigin(UMoviePipelinePrimaryConfig* InP
 	{
 		Configuration->CopyFrom(InPreset);
 		PresetOrigin = TSoftObjectPtr<UMoviePipelinePrimaryConfig>(InPreset);
+		ConfigMode = EMoviePipelineConfigMode::Preset;
 	}
 }
 
@@ -356,6 +399,45 @@ void UMoviePipelineExecutorJob::OnGraphPreSave(UObject* InObject, FObjectPreSave
 	RefreshAllVariableAssignments();
 }
 
+UMoviePipelineBasicConfig* UMoviePipelineExecutorJob::SetupBasicConfiguration()
+{
+#if WITH_EDITOR
+	Modify();
+#endif
+
+	// If the user has saved a default Basic config, duplicate it; otherwise create a fresh one.
+	if (!BasicConfig)
+	{
+		if (const UMoviePipelineBasicConfig* SavedDefault = UMoviePipelineBasicConfig::GetSavedDefault())
+		{
+			BasicConfig = DuplicateObject(SavedDefault, this);
+		}
+		else
+		{
+			BasicConfig = NewObject<UMoviePipelineBasicConfig>(this);
+		}
+	}
+
+	ConfigMode = EMoviePipelineConfigMode::Basic;
+
+	// Ensure all child shots have a BasicShotConfig so they can display override controls.
+	// SetupBasicShotOverrides() is idempotent — safe to call if already initialized.
+	for (const TObjectPtr<UMoviePipelineExecutorShot>& Shot : ShotInfo)
+	{
+		if (Shot)
+		{
+			Shot->SetupBasicShotOverrides();
+		}
+	}
+
+	return BasicConfig;
+}
+
+bool UMoviePipelineExecutorJob::CanUseConsoleVariableOverrides() const
+{
+	return IsUsingGraphConfiguration() || IsUsingBasicConfiguration();
+}
+
 void UMoviePipelineExecutorShot::SetGraphPreset(const UMovieGraphConfig* InGraphPreset, const bool bUpdateVariableAssignments)
 {
 #if WITH_EDITOR
@@ -453,10 +535,41 @@ bool UMoviePipelineExecutorShot::IsUsingGraphConfiguration() const
 {
 	const UMoviePipelineExecutorJob* PrimaryJob = GetTypedOuter<UMoviePipelineExecutorJob>();
 
-	// Calling GetGraphPreset() here is important, rather than just referencing GraphPreset (to ensure that the soft ptr has a chance to load) 
+	// Calling GetGraphPreset() here is important, rather than just referencing GraphPreset (to ensure that the soft ptr has a chance to load)
 	return (GetGraphPreset() != nullptr) || (PrimaryJob && PrimaryJob->IsUsingGraphConfiguration());
 }
 
+bool UMoviePipelineExecutorShot::IsUsingBasicConfiguration() const
+{
+	const UMoviePipelineExecutorJob* PrimaryJob = GetTypedOuter<UMoviePipelineExecutorJob>();
+	return PrimaryJob && PrimaryJob->IsUsingBasicConfiguration();
+}
+
+UMoviePipelineBasicConfig* UMoviePipelineExecutorShot::SetupBasicShotOverrides()
+{
+#if WITH_EDITOR
+	Modify();
+#endif
+
+	if (!BasicShotConfig)
+	{
+		BasicShotConfig = NewObject<UMoviePipelineBasicConfig>(this);
+
+		// These inner conditions must be true for their dependent properties to be editable
+		// in the details panel. For shot configs, editability is controlled by the bOverride_*
+		// flags; these inner bools are always-on so they don't interfere.
+		BasicShotConfig->bEnableCustomStartFrame = true;
+		BasicShotConfig->bEnableCustomEndFrame = true;
+		BasicShotConfig->bPathTracedEnableDenoiser = true;
+	}
+	return BasicShotConfig;
+}
+
+bool UMoviePipelineExecutorShot::CanUseConsoleVariableOverrides() const
+{
+	return IsUsingGraphConfiguration() || IsUsingBasicConfiguration();
+}
+
 void UMoviePipelineExecutorShot::PreSave(FObjectPreSaveContext ObjectSaveContext)
 {
 	Super::PreSave(ObjectSaveContext);
diff --git a/Engine/Plugins/MovieScene/MovieRenderPipeline/Source/MovieRenderPipelineCore/Public/MoviePipelineBasicConfig.h b/Engine/Plugins/MovieScene/MovieRenderPipeline/Source/MovieRenderPipelineCore/Public/MoviePipelineBasicConfig.h
new file mode 100644
--- /dev/null
+++ b/Engine/Plugins/MovieScene/MovieRenderPipeline/Source/MovieRenderPipelineCore/Public/MoviePipelineBasicConfig.h
@@ -0,0 +1,200 @@
+// Copyright Epic Games, Inc. All Rights Reserved.
+
+#pragma once
+
+#include "Graph/MovieGraphNamedResolution.h"
+#include "Graph/Nodes/MovieGraphFileOutputNode.h"
+#include "UObject/Object.h"
+#include "UObject/SoftObjectPath.h"
+
+#include "MoviePipelineBasicConfig.generated.h"
+
+#define UE_API MOVIERENDERPIPELINECORE_API
+
+class UMovieGraphConfig;
+
+/** Denoiser type for the Path Traced Renderer in Basic configuration mode. */
+UENUM(BlueprintType)
+enum class EMoviePipelineBasicDenoiserType : uint8
+{
+	Spatial = 0,
+	Temporal = 1
+};
+
+/**
+ * A simplified configuration for Movie Render Queue jobs. Provides a pared-down set of rendering
+ * options that generates a UMovieGraphConfig just-in-time before rendering. This allows most users
+ * to configure renders without touching Presets or Graphs, while still using the full graph pipeline
+ * under the hood.
+ */
+UCLASS(MinimalAPI, BlueprintType)
+class UMoviePipelineBasicConfig : public UObject
+{
+	GENERATED_BODY()
+
+public:
+	UE_API UMoviePipelineBasicConfig();
+
+	/**
+	 * Merge a shot's basic config with this config, providing the result in "OutMerged".
+	 * For each bOverride_* that is true within InShotOverrides, use the shot's value; otherwise, use the value from this config.
+	 */
+	UE_API void MergeOverrides(const UMoviePipelineBasicConfig* InShotOverrides, UMoviePipelineBasicConfig& OutMerged) const;
+
+	/**
+	 * Generate a UMovieGraphConfig from this config (or a merged config).
+	 * Note that EnforceInvariants() will be called on InConfig to ensure the generated graph is valid.
+	 */
+	static UE_API UMovieGraphConfig* GenerateGraph(UMoviePipelineBasicConfig* InConfig, UObject* InOuter);
+
+	/**
+	 * Returns true if at least one bOverride_* flag is set on this config.
+	 */
+	UE_API bool HasAnyOverrides() const;
+
+	/**
+	 * Enforce minimum-valid-state invariants:
+	 *   - At least one renderer (Deferred or Path Traced) must be enabled; if neither is, Deferred is re-enabled.
+	 *   - EnabledOutputTypes must be non-empty; if it is, PNG is restored as the default.
+	 */
+	UE_API void EnforceInvariants();
+
+#if WITH_EDITOR
+	virtual UE_API void PostEditChangeProperty(FPropertyChangedEvent& PropertyChangedEvent) override;
+
+	/** Save InConfig as the user's default Basic config (persisted to a local uasset in the project). */
+	static UE_API void SaveAsDefault(const UMoviePipelineBasicConfig* InConfig);
+#endif
+
+	/** Return the user's saved default Basic config, or nullptr if none exists. */
+	static UE_API const UMoviePipelineBasicConfig* GetSavedDefault();
+
+public:
+	// --- Override flags (most are relevant for shots only, with a few exceptions) ---
+
+	UPROPERTY(EditAnywhere, BlueprintReadWrite, Category = Overrides, meta = (InlineEditConditionToggle))
+	uint8 bOverride_OutputDirectory : 1;
+
+	UPROPERTY(EditAnywhere, BlueprintReadWrite, Category = Overrides, meta = (InlineEditConditionToggle))
+	uint8 bOverride_FileNameFormat : 1;
+
+	UPROPERTY(EditAnywhere, BlueprintReadWrite, Category = Overrides, meta = (InlineEditConditionToggle))
+	uint8 bOverride_OutputResolution : 1;
+
+	UPROPERTY(EditAnywhere, BlueprintReadWrite, Category = Overrides, meta = (InlineEditConditionToggle))
+	uint8 bOverride_EnabledOutputTypes : 1;
+
+	UPROPERTY(EditAnywhere, BlueprintReadWrite, Category = Overrides, meta = (InlineEditConditionToggle))
+	uint8 bOverride_bEnableCustomStartFrame : 1;
+
+	UPROPERTY(EditAnywhere, BlueprintReadWrite, Category = Overrides, meta = (InlineEditConditionToggle))
+	uint8 bOverride_CustomStartFrame : 1;
+
+	UPROPERTY(EditAnywhere, BlueprintReadWrite, Category = Overrides, meta = (InlineEditConditionToggle))
+	uint8 bOverride_bEnableCustomEndFrame : 1;
+
+	UPROPERTY(EditAnywhere, BlueprintReadWrite, Category = Overrides, meta = (InlineEditConditionToggle))
+	uint8 bOverride_CustomEndFrame : 1;
+
+	UPROPERTY(EditAnywhere, BlueprintReadWrite, Category = Overrides, meta = (InlineEditConditionToggle))
+	uint8 bOverride_bUseDeferredRenderer : 1;
+
+	UPROPERTY(EditAnywhere, BlueprintReadWrite, Category = Overrides, meta = (InlineEditConditionToggle))
+	uint8 bOverride_DeferredSpatialSampleCount : 1;
+
+	UPROPERTY(EditAnywhere, BlueprintReadWrite, Category = Overrides, meta = (InlineEditConditionToggle))
+	uint8 bOverride_bUsePathTracedRenderer : 1;
+
+	UPROPERTY(EditAnywhere, BlueprintReadWrite, Category = Overrides, meta = (InlineEditConditionToggle))
+	uint8 bOverride_PathTracedSpatialSampleCount : 1;
+
+	UPROPERTY(EditAnywhere, BlueprintReadWrite, Category = Overrides, meta = (InlineEditConditionToggle))
+	uint8 bOverride_bPathTracedEnableDenoiser : 1;
+
+	UPROPERTY(EditAnywhere, BlueprintReadWrite, Category = Overrides, meta = (InlineEditConditionToggle))
+	uint8 bOverride_PathTracedDenoiserType : 1;
+
+	UPROPERTY(EditAnywhere, BlueprintReadWrite, Category = Overrides, meta = (InlineEditConditionToggle))
+	uint8 bOverride_NumWarmUpFrames : 1;
+
+	UPROPERTY(EditAnywhere, BlueprintReadWrite, Category = Overrides, meta = (InlineEditConditionToggle))
+	uint8 bOverride_TemporalSampleCount : 1;
+
+	UPROPERTY(EditAnywhere, BlueprintReadWrite, Category = Overrides, meta = (InlineEditConditionToggle))
+	uint8 bOverride_BurnInClass : 1;
+
+	// --- Properties ---
+
+	/** The directory all the output files are relative to. */
+	UPROPERTY(EditAnywhere, BlueprintReadWrite, Category = "Output")
+	FDirectoryPath OutputDirectory;
+
+	/** The format string for file names. Can include folder prefixes, and format string tokens ({shot_name}, etc.) */
+	UPROPERTY(EditAnywhere, BlueprintReadWrite, Category = "Output")
+	FString FileNameFormat;
+
+	/** The resolution that output files are exported at. */
+	UPROPERTY(EditAnywhere, BlueprintReadWrite, Category = "Output")
+	FMovieGraphNamedResolution OutputResolution;
+
+	/** The output types that rendered media should be saved to. */
+	UPROPERTY(EditAnywhere, BlueprintReadWrite, Category = "Output")
+	TArray<TSoftClassPtr<UMovieGraphFileOutputNode>> EnabledOutputTypes;
+
+	/** Enable a custom start frame, overriding the sequence's playback range start. */
+	UPROPERTY(EditAnywhere, BlueprintReadWrite, Category = "Output", meta = (InlineEditConditionToggle))
+	bool bEnableCustomStartFrame;
+
+	/** A custom start frame (overrides the sequence's playback range start). Only used if bEnableCustomStartFrame is true. */
+	UPROPERTY(EditAnywhere, BlueprintReadWrite, Category = "Output", meta = (EditCondition = "bEnableCustomStartFrame"))
+	int32 CustomStartFrame;
+
+	/** Enable a custom end frame, overriding the sequence's playback range end. */
+	UPROPERTY(EditAnywhere, BlueprintReadWrite, Category = "Output", meta = (InlineEditConditionToggle))
+	bool bEnableCustomEndFrame;
+
+	/** A custom end frame (overrides the sequence's playback range end). Only used if bEnableCustomEndFrame is true. */
+	UPROPERTY(EditAnywhere, BlueprintReadWrite, Category = "Output", meta = (EditCondition = "bEnableCustomEndFrame"))
+	int32 CustomEndFrame;
+
+	/** Whether the deferred renderer should be enabled. */
+	UPROPERTY(EditAnywhere, BlueprintReadWrite, Category = "Rendering")
+	bool bUseDeferredRenderer;
+
+	/** Spatial sample count for the Deferred Renderer. */
+	UPROPERTY(EditAnywhere, BlueprintReadWrite, Category = "Rendering", meta = (DisplayName = "Spatial Sample Count", EditCondition = "bUseDeferredRenderer", UIMin = 1, ClampMin = 1))
+	int32 DeferredSpatialSampleCount;
+
+	/** Whether the path tracer should be enabled. */
+	UPROPERTY(EditAnywhere, BlueprintReadWrite, Category = "Rendering")
+	bool bUsePathTracedRenderer;
+
+	/** Spatial sample count for the Path Traced Renderer. */
+	UPROPERTY(EditAnywhere, BlueprintReadWrite, Category = "Rendering", meta = (DisplayName = "Spatial Sample Count", EditCondition = "bUsePathTracedRenderer", UIMin = 1, ClampMin = 1))
+	int32 PathTracedSpatialSampleCount;
+
+	/** Whether the result from the path tracer should be passed through a denoiser. */
+	UPROPERTY(EditAnywhere, BlueprintReadWrite, Category = "Rendering", meta = (InlineEditConditionToggle))
+	bool bPathTracedEnableDenoiser;
+
+	/** If denoising is enabled, this specifies the specific denoiser that path tracer results should be passed to. */
+	UPROPERTY(EditAnywhere, BlueprintReadWrite, Category = "Rendering", meta = (DisplayName = "Denoiser", EditCondition = "bPathTracedEnableDenoiser"))
+	EMoviePipelineBasicDenoiserType PathTracedDenoiserType;
+
+	/** The number of frames to run (without producing output) at the start of each shot to warm up systems like Niagara, Chaos, etc. */
+	UPROPERTY(EditAnywhere, BlueprintReadWrite, Category = "Rendering", meta = (DisplayName = "Num Warm Up Frames", UIMin = 0, ClampMin = 0))
+	int32 NumWarmUpFrames;
+
+	/** The number of temporal sub-samples per output frame. Higher values will reduce temporal aliasing. */
+	UPROPERTY(EditAnywhere, BlueprintReadWrite, Category = "Rendering", meta = (DisplayName = "Temporal Sample Count", UIMin = 1, ClampMin = 1))
+	int32 TemporalSampleCount;
+
+	/**
+	 * The burn-in widget class to use on all output types that support burn-in. Only applies if the burn-in override is enabled.
+	 * Burn-ins are enabled for all output types that support them, and compositing is enabled for all output types that support burn-in compositing.
+	 */
+	UPROPERTY(EditAnywhere, BlueprintReadWrite, Category = "Output", meta = (EditCondition = "bOverride_BurnInClass", MetaClass = "/Script/MovieRenderPipelineCore.MovieGraphBurnInWidget"))
+	FSoftClassPath BurnInClass;
+};
+
+#undef UE_API
diff --git a/Engine/Plugins/MovieScene/MovieRenderPipeline/Source/MovieRenderPipelineCore/Public/MoviePipelineQueue.h b/Engine/Plugins/MovieScene/MovieRenderPipeline/Source/MovieRenderPipelineCore/Public/MoviePipelineQueue.h
--- a/Engine/Plugins/MovieScene/MovieRenderPipeline/Source/MovieRenderPipelineCore/Public/MoviePipelineQueue.h
+++ b/Engine/Plugins/MovieScene/MovieRenderPipeline/Source/MovieRenderPipelineCore/Public/MoviePipelineQueue.h
@@ -4,6 +4,7 @@
 #include "UObject/Object.h"
 #include "MovieRenderPipelineDataTypes.h"
 #include "MovieJobVariableAssignmentContainer.h"
+#include "MoviePipelineBasicConfig.h"
 #include "MoviePipelinePrimaryConfig.h"
 #include "MoviePipelineShotConfig.h"
 #include "MoviePipelineConfigBase.h"
@@ -58,6 +59,20 @@ public:
 	EMoviePipelinePostRenderActionType PostRenderActionType = EMoviePipelinePostRenderActionType::DoNothing;
 };
 
+/** The active configuration mode for a Movie Render Queue job. */
+UENUM(BlueprintType)
+enum class EMoviePipelineConfigMode : uint8
+{
+	/** Legacy UMoviePipelinePrimaryConfig preset. */
+	Preset,
+
+	/** Graph-based UMovieGraphConfig configuration. */
+	Graph,
+
+	/** Simplified UMoviePipelineBasicConfig configuration. */
+	Basic
+};
+
 /**
 * This class represents a segment of work within the Executor Job. This should be owned
 * by the UMoviePipelineExecutorJob and can be created before the movie pipeline starts to
@@ -160,7 +175,9 @@ public:
 	UE_API bool IsUsingGraphConfiguration() const;
 
 	/**
-	 * Gets the graph-style preset that this job is using. If the job is not using a graph-style preset, returns nullptr.
+	 * Gets the graph-style preset that this shot is using.
+	 * A non-null graph config does not imply the graph config is the config that's in use;
+	 * reference GetConfigMode() on the job for the config type that's active.
 	 */
 	UFUNCTION(BlueprintPure, Category = "Movie Render Pipeline")
 	UMovieGraphConfig* GetGraphPreset() const
@@ -186,7 +203,7 @@ public:
 	* 
 	* If you specify true for bIsForPrimaryOverrides it returns an object that allows this shot to override a variable that comes from the primary 
 	* graph. If you return false, then it returns an object that allows overriding a variable for this shot's override config (see: GetGraphPreset).
-	* See UMoviePipelineExecutorJob's version of this functoin for more details.
+	* See UMoviePipelineExecutorJob's version of this function for more details.
 	* 
 	* @param InGraph - The graph asset to return the config for. If this shot has its own Graph Preset override, you should return GetGraphPreset()
 	* or one of it's sub-graph pointers. If this shot is just trying to override the Primary Graph from the parent UMoviePipelineExecutorJob then
@@ -202,7 +219,27 @@ public:
 	UFUNCTION(BlueprintCallable, Category = "Movie Render Pipeline", DisplayName = "Get or Create Variable Overrides", meta = (ScriptName = "GetOrCreateVariableOverrides"))
 	UE_API UMovieJobVariableAssignmentContainer* GetOrCreateJobVariableAssignmentsForGraph(const UMovieGraphConfig* InGraph, const bool bIsForPrimaryOverrides = false);
 
-	/** Returns whether this should should be rendered */
+	/** Returns true if the parent job is using basic configuration mode. */
+	UFUNCTION(BlueprintPure, Category = "Movie Render Pipeline")
+	UE_API bool IsUsingBasicConfiguration() const;
+
+	/**
+	 * Gets the basic shot config overrides. Returns nullptr if no shot overrides are set.
+	 * A non-null Basic config override does not imply the Basic config is the config that's in use;
+	 * reference the job's GetConfigMode() for the config type that's active.
+	 */
+	UFUNCTION(BlueprintPure, Category = "Movie Render Pipeline")
+	UMoviePipelineBasicConfig* GetBasicShotConfig() const { return BasicShotConfig; }
+
+	/** Creates and returns basic shot overrides for this shot. */
+	UFUNCTION(BlueprintCallable, Category = "Movie Render Pipeline")
+	UE_API UMoviePipelineBasicConfig* SetupBasicShotOverrides();
+
+	/** Returns true if console variable overrides can be used in the active configuration. */
+	UFUNCTION(BlueprintPure, Category = "Movie Render Pipeline")
+	UE_API bool CanUseConsoleVariableOverrides() const;
+
+	/** Returns whether this should be rendered. */
 	UFUNCTION(BlueprintPure, Category = "Movie Render Pipeline")
 	bool ShouldRender() const
 	{
@@ -277,9 +314,9 @@ public:
 	FOnMoviePipelineShotGraphPresetChanged OnShotGraphPresetChanged;
 
 	/**
-	 * (Optional) Console variable overrides which are applied after cvars set via nodes. Only applies to graph-based configs.
+	 * (Optional) Console variable overrides which are applied after cvars set via nodes. Applies to graph-based and Basic configurations.
 	 */
-	UPROPERTY(BlueprintReadWrite, EditAnywhere, Category = "Console Variable Overrides", meta = (DisplayName = "Console Variables", EditCondition="IsUsingGraphConfiguration()", EditConditionHides))
+	UPROPERTY(BlueprintReadWrite, EditAnywhere, Category = "Console Variable Overrides", meta = (DisplayName = "Console Variables", EditCondition="CanUseConsoleVariableOverrides()", EditConditionHides))
 	TArray<FMoviePipelineConsoleVariableEntry> ConsoleVariableOverrides;
 
 public:
@@ -310,6 +347,10 @@ private:
 	/** Overrides on the variables in the primary graph (and its subgraphs) associated with this job. */
 	UPROPERTY(EditAnywhere, Instanced, Category = "Variable Assignments")
 	TArray<TObjectPtr<UMovieJobVariableAssignmentContainer>> PrimaryGraphVariableAssignments;
+
+	/** The basic configuration overrides for this shot. Non-null indicates shot-level overrides in Basic configuration mode. */
+	UPROPERTY(Instanced)
+	TObjectPtr<UMoviePipelineBasicConfig> BasicShotConfig;
 };
 
 DECLARE_MULTICAST_DELEGATE_TwoParams(FOnMoviePipelineJobGraphPresetChanged, UMoviePipelineExecutorJob*, UMovieGraphConfig*);
@@ -477,6 +518,10 @@ public:
 	 * Gets the configuration for the job. This configuration is either standalone (not associated with any preset), or
 	 * contains a copy of the preset origin plus any modifications made on top of it. If the preset that this
 	 * configuration was originally based on no longer exists, this configuration will still be valid.
+	 * 
+	 * Note that a non-null configuration does not imply that this Preset is the active configuration;
+	 * reference GetConfigMode() for the config type that's active.
+	 * 
 	 * @see GetPresetOrigin()
 	 */
 	UFUNCTION(BlueprintPure, Category = "Movie Render Pipeline")
@@ -488,19 +533,43 @@ public:
 	UFUNCTION(BlueprintCallable, Category = "Movie Render Pipeline")
 	UE_API void SetConfiguration(UMoviePipelinePrimaryConfig* InPreset);
 
+	/** Returns the active configuration mode for this job. */
+	UFUNCTION(BlueprintPure, Category = "Movie Render Pipeline")
+	EMoviePipelineConfigMode GetConfigMode() const { return ConfigMode; }
+
 	/**
 	 * Returns true if this job is using graph-style configuration, else false.
 	 */
 	UFUNCTION(BlueprintPure, Category = "Movie Render Pipeline")
-	bool IsUsingGraphConfiguration() const
-	{
-		// Calling GetGraphPreset() here is important, rather than just referencing GraphPreset (to ensure that the soft ptr has a chance to load)
-		return GetGraphPreset() != nullptr;
-	}
+	bool IsUsingGraphConfiguration() const { return ConfigMode == EMoviePipelineConfigMode::Graph; }
+
+	/** Returns true if this job is using basic configuration mode. */
+	UFUNCTION(BlueprintPure, Category = "Movie Render Pipeline")
+	bool IsUsingBasicConfiguration() const { return ConfigMode == EMoviePipelineConfigMode::Basic; }
 
 	/**
-	 * Gets the graph-style preset that this job is using. If the job is not using a graph-style preset, returns nullptr.
-	 * @see GetPresetOrigin()
+	 * Gets the basic configuration object, or nullptr if basic configuration has not been set up.
+	 * A non-null Basic config does not imply the Basic config is the config that's in use;
+	 * reference GetConfigMode() for the config type that's active.
+	 */
+	UFUNCTION(BlueprintPure, Category = "Movie Render Pipeline")
+	UMoviePipelineBasicConfig* GetBasicConfig() const { return BasicConfig; }
+
+	/**
+	 * Sets up basic configuration mode for this job. Note that this will cause the job to switch over to
+	 * using the Basic configuration. Returns the created basic config object.
+	 */
+	UFUNCTION(BlueprintCallable, Category = "Movie Render Pipeline")
+	UE_API UMoviePipelineBasicConfig* SetupBasicConfiguration();
+
+	/** Returns true if console variable overrides can be used in the active configuration. */
+	UFUNCTION(BlueprintPure, Category = "Movie Render Pipeline")
+	UE_API bool CanUseConsoleVariableOverrides() const;
+
+	/**
+	 * Gets the graph-style preset assigned to this job.
+	 * A non-null graph config does not imply the graph config is the config that's in use;
+	 * reference GetConfigMode() for the config type that's active.
 	 */
 	UFUNCTION(BlueprintPure, Category = "Movie Render Pipeline")
 	UMovieGraphConfig* GetGraphPreset() const
@@ -545,6 +614,7 @@ public:
 
 public:
 	// UObject Interface
+	UE_API virtual void Serialize(FArchive& Ar) override;
 #if WITH_EDITOR
 	UE_API void PostEditChangeProperty(FPropertyChangedEvent& PropertyChangedEvent);
 #endif
@@ -592,11 +662,11 @@ public:
 	FString JobName;
 
 	/** Which sequence should this job render? */
-	UPROPERTY(BlueprintReadWrite, EditAnywhere, BlueprintSetter = "SetSequence", Category = "Movie Render Pipeline", meta = (AllowedClasses = "/Script/LevelSequence.LevelSequence"))
+	UPROPERTY(BlueprintReadWrite, EditAnywhere, BlueprintSetter = "SetSequence", Category = "Movie Render Pipeline", meta = (AllowedClasses = "/Script/LevelSequence.LevelSequence", DisplayThumbnail = false))
 	FSoftObjectPath Sequence;
 
 	/** Which map should this job render on */
-	UPROPERTY(BlueprintReadWrite, EditAnywhere, Category = "Movie Render Pipeline", meta = (AllowedClasses = "/Script/Engine.World"))
+	UPROPERTY(BlueprintReadWrite, EditAnywhere, Category = "Movie Render Pipeline", meta = (AllowedClasses = "/Script/Engine.World", DisplayThumbnail = false))
 	FSoftObjectPath Map;
 
 	/** (Optional) If left blank, will default to system username. Can be shown in burn in as a first point of contact about the content. */
@@ -623,9 +693,9 @@ public:
 	FOnMoviePipelineJobGraphPresetChanged OnJobGraphPresetChanged;
 
 	/**
-	 * (Optional) Console variable overrides which are applied after cvars set via nodes. Only applies to graph-based configs.
+	 * (Optional) Console variable overrides which are applied after cvars set via nodes. Applies to graph-based and Basic configurations.
 	 */
-	UPROPERTY(BlueprintReadWrite, EditAnywhere, Category = "Console Variable Overrides", meta = (DisplayName = "Console Variables", EditCondition="IsUsingGraphConfiguration()", EditConditionHides))
+	UPROPERTY(BlueprintReadWrite, EditAnywhere, Category = "Console Variable Overrides", meta = (DisplayName = "Console Variables", EditCondition="CanUseConsoleVariableOverrides()", EditConditionHides))
 	TArray<FMoviePipelineConsoleVariableEntry> ConsoleVariableOverrides;
 
 private:
@@ -652,13 +722,25 @@ private:
 	UPROPERTY()
 	bool bEnabled;
 
-	/** The graph-based configuration preset that this job is using. Can be nullptr. */
+	/**
+	 * The graph-based configuration preset that this job is using. Can be nullptr. Typically, this should be set with
+	 * SetGraphPreset() in order to keep the ConfigMode in sync. May be non-null even when not in Graph mode
+	 * (it's preserved across mode switches). Check ConfigMode to determine the active configuration.
+	 */
 	UPROPERTY()
 	TSoftObjectPtr<UMovieGraphConfig> GraphPreset;
 
 	/** Overrides on the variables in the graph (and subgraphs) associated with this job. */
 	UPROPERTY(EditAnywhere, Instanced, Category = "Variable Assignments")
 	TArray<TObjectPtr<UMovieJobVariableAssignmentContainer>> GraphVariableAssignments;
+
+	/** The basic configuration for this job. May be non-null even when not in Basic mode (preserved across mode switches). */
+	UPROPERTY(Instanced)
+	TObjectPtr<UMoviePipelineBasicConfig> BasicConfig;
+
+	/** The active configuration mode for this job. Determines which config object is in use. */
+	UPROPERTY()
+	EMoviePipelineConfigMode ConfigMode = EMoviePipelineConfigMode::Preset;
 };
 
 /**
diff --git a/Engine/Plugins/MovieScene/MovieRenderPipeline/Source/MovieRenderPipelineCore/Public/MovieRenderPipelineCoreObjectVersion.h b/Engine/Plugins/MovieScene/MovieRenderPipeline/Source/MovieRenderPipelineCore/Public/MovieRenderPipelineCoreObjectVersion.h
--- a/Engine/Plugins/MovieScene/MovieRenderPipeline/Source/MovieRenderPipelineCore/Public/MovieRenderPipelineCoreObjectVersion.h
+++ b/Engine/Plugins/MovieScene/MovieRenderPipeline/Source/MovieRenderPipelineCore/Public/MovieRenderPipelineCoreObjectVersion.h
@@ -17,6 +17,9 @@ struct FMovieRenderPipelineCoreObjectVersion
 		/** Added bOnlyMatchComponents to some conditions within graph collections, defaulting to true. */
 		OnlyMatchComponentsAdded,
 
+		/** Added EMoviePipelineConfigMode to UMoviePipelineExecutorJob. Pre-existing assets derive the mode from pointer state in Serialize(). */
+		JobConfigModeAdded,
+
 		// --- Add new versions above this line ---
 		VersionPlusOne,
 		LatestVersion = VersionPlusOne - 1
diff --git a/Engine/Plugins/MovieScene/MovieRenderPipeline/Source/MovieRenderPipelineEditor/Private/MoviePipelineEditorBlueprintLibrary.cpp b/Engine/Plugins/MovieScene/MovieRenderPipeline/Source/MovieRenderPipelineEditor/Private/MoviePipelineEditorBlueprintLibrary.cpp
--- a/Engine/Plugins/MovieScene/MovieRenderPipeline/Source/MovieRenderPipelineEditor/Private/MoviePipelineEditorBlueprintLibrary.cpp
+++ b/Engine/Plugins/MovieScene/MovieRenderPipeline/Source/MovieRenderPipelineEditor/Private/MoviePipelineEditorBlueprintLibrary.cpp
@@ -1,10 +1,13 @@
 // Copyright Epic Games, Inc. All Rights Reserved.
 
 #include "MoviePipelineEditorBlueprintLibrary.h"
+
+#include "ContentBrowserModule.h"
 #include "UObject/Package.h"
 #include "Misc/PackageName.h"
 #include "AssetRegistry/AssetRegistryModule.h"
 #include "FileHelpers.h"
+#include "IContentBrowserSingleton.h"
 #include "ObjectTools.h"
 #include "PackageTools.h"
 #include "LevelSequence.h"
@@ -18,11 +21,14 @@
 #include "Settings/EditorLoadingSavingSettings.h"
 #include "Editor.h"
 #include "GenericPlatform/GenericPlatformProcess.h"
+#include "ScopedTransaction.h"
 #include "SequencerUtilities.h"
 #include "MoviePipelineBlueprintLibrary.h"
 #include "MoviePipelineOutputSetting.h"
 #include "Graph/MovieGraphBlueprintLibrary.h"
+#include "Graph/MovieGraphConfig.h"
 #include "Graph/Nodes/MovieGraphGlobalOutputSettingNode.h"
+#include "MoviePipelineBasicConfig.h"
 
 #include UE_INLINE_GENERATED_CPP_BY_NAME(MoviePipelineEditorBlueprintLibrary)
 
@@ -338,5 +344,81 @@ FString UMoviePipelineEditorBlueprintLibrary::ResolveOutputDirectoryFromJob(UMov
 	return OutResolvedPath;
 }
 
+bool UMoviePipelineEditorBlueprintLibrary::SaveBasicConfigAsGraphConfig(UMoviePipelineExecutorJob* InJob)
+{
+	if (!InJob || !InJob->IsUsingBasicConfiguration())
+	{
+		return false;
+	}
+
+	UMoviePipelineBasicConfig* BasicConfig = InJob->GetBasicConfig();
+	if (!BasicConfig)
+	{
+		return false;
+	}
+
+	UMovieGraphConfig* GeneratedGraph = UMoviePipelineBasicConfig::GenerateGraph(BasicConfig, GetTransientPackage());
+	if (!GeneratedGraph)
+	{
+		return false;
+	}
+
+	// Show Save As dialog for the user to choose asset path/name
+	FSaveAssetDialogConfig SaveDialogConfig;
+	SaveDialogConfig.DialogTitleOverride = LOCTEXT("SaveAsGraph_DialogTitle", "Save As Movie Render Graph Config");
+	SaveDialogConfig.DefaultPath = TEXT("/Game/");
+	SaveDialogConfig.DefaultAssetName = TEXT("NewMovieRenderGraphConfig");
+	SaveDialogConfig.ExistingAssetPolicy = ESaveAssetDialogExistingAssetPolicy::AllowButWarn;
+
+	const FContentBrowserModule& ContentBrowserModule = FModuleManager::LoadModuleChecked<FContentBrowserModule>("ContentBrowser");
+	const FString SaveObjectPath = ContentBrowserModule.Get().CreateModalSaveAssetDialog(SaveDialogConfig);
+	if (SaveObjectPath.IsEmpty())
+	{
+		return false; // User canceled
+	}
+
+	// Create package, duplicate generated graph into it, save
+	const FString NewPackageName = FPackageName::ObjectPathToPackageName(SaveObjectPath);
+	const FString NewAssetName = FPaths::GetBaseFilename(NewPackageName);
+
+	UPackage* NewPackage = CreatePackage(*NewPackageName);
+	NewPackage->MarkAsFullyLoaded();
+
+	// The generated graph was built by programmatically modifying a duplicate of the default graph via
+	// UMovieGraphConfig APIs (InsertBefore, AddLabeledEdge, etc.). These APIs only update the runtime
+	// representation; the duplicated PipelineEdGraph still reflects the original default graph's connections.
+	// Clear it so the editor rebuilds a fresh editor graph from the authoritative runtime data when opened.
+	// See UMoviePipelineEdGraph::InitFromRuntimeGraph() for the stale GraphNode cleanup that complements this.
+	if (GeneratedGraph->PipelineEdGraph)
+	{
+		GeneratedGraph->PipelineEdGraph->Rename(nullptr, GetTransientPackage(), REN_DontCreateRedirectors | REN_NonTransactional);
+		GeneratedGraph->PipelineEdGraph = nullptr;
+	}
+
+	UMovieGraphConfig* SavedGraph = Cast<UMovieGraphConfig>(
+		StaticDuplicateObject(GeneratedGraph, NewPackage, FName(*NewAssetName)));
+	if (!SavedGraph)
+	{
+		return false;
+	}
+	SavedGraph->SetFlags(RF_Public | RF_Transactional | RF_Standalone);
+	SavedGraph->MarkPackageDirty();
+
+	TArray<UPackage*> PackagesToSave;
+	PackagesToSave.Add(NewPackage);
+	if (!UEditorLoadingAndSavingUtils::SavePackages(PackagesToSave, true))
+	{
+		return false;
+	}
+
+	FAssetRegistryModule::AssetCreated(SavedGraph);
+
+	// Switch the job to Graph mode using the saved asset
+	FScopedTransaction Transaction(LOCTEXT("SaveBasicAsGraph_Transaction", "Save Basic Config as Graph Config"));
+	InJob->Modify();
+	InJob->SetGraphPreset(SavedGraph);
+
+	return true;
+}
 
 #undef LOCTEXT_NAMESPACE // "MoviePipelineEditorBlueprintLibrary"
diff --git a/Engine/Plugins/MovieScene/MovieRenderPipeline/Source/MovieRenderPipelineEditor/Private/MoviePipelineNewProcessExecutor.cpp b/Engine/Plugins/MovieScene/MovieRenderPipeline/Source/MovieRenderPipelineEditor/Private/MoviePipelineNewProcessExecutor.cpp
--- a/Engine/Plugins/MovieScene/MovieRenderPipeline/Source/MovieRenderPipelineEditor/Private/MoviePipelineNewProcessExecutor.cpp
+++ b/Engine/Plugins/MovieScene/MovieRenderPipeline/Source/MovieRenderPipelineEditor/Private/MoviePipelineNewProcessExecutor.cpp
@@ -103,14 +103,19 @@ void UMoviePipelineNewProcessExecutor::Execute_Implementation(UMoviePipelineQueu
 			continue;
 		}
 
+		// Determine the graph to evaluate. For Basic mode, generate a transient graph from the Basic config.
+		UMovieGraphConfig* GraphConfig = nullptr;
 		if (Job->IsUsingGraphConfiguration())
 		{
-			UMovieGraphConfig* GraphConfig = Job->GetGraphPreset();
-			if (!GraphConfig)
-			{
-				continue;
-			}
-			
+			GraphConfig = Job->GetGraphPreset();
+		}
+		else if (Job->IsUsingBasicConfiguration())
+		{
+			GraphConfig = UMoviePipelineBasicConfig::GenerateGraph(Job->GetBasicConfig(), GetTransientPackage());
+		}
+
+		if (GraphConfig)
+		{
 			FMovieGraphTraversalContext TraversalContext;
 			TraversalContext.Job = Job;
 
diff --git a/Engine/Plugins/MovieScene/MovieRenderPipeline/Source/MovieRenderPipelineEditor/Private/MoviePipelinePIEExecutor.cpp b/Engine/Plugins/MovieScene/MovieRenderPipeline/Source/MovieRenderPipelineEditor/Private/MoviePipelinePIEExecutor.cpp
--- a/Engine/Plugins/MovieScene/MovieRenderPipeline/Source/MovieRenderPipelineEditor/Private/MoviePipelinePIEExecutor.cpp
+++ b/Engine/Plugins/MovieScene/MovieRenderPipeline/Source/MovieRenderPipelineEditor/Private/MoviePipelinePIEExecutor.cpp
@@ -156,6 +156,16 @@ void UMoviePipelinePIEExecutor::Start(const UMoviePipelineExecutorJob* InJob)
 				UE_LOG(LogMovieRenderPipeline, Warning, TEXT("Failed to evaluate graph to determine output resolution. Traversal error: %s"), *OutTraversalError);
 			}
 		}
+		else if (InJob->IsUsingBasicConfiguration())
+		{
+			// Basic-mode jobs generate their graph later in DuplicateJobAndConfiguration().
+			// Use the BasicConfig's output resolution directly for the PIE window size.
+			if (const UMoviePipelineBasicConfig* BasicCfg = InJob->GetBasicConfig())
+			{
+				const FMovieGraphNamedResolution& Resolution = BasicCfg->OutputResolution;
+				WindowSize = FVector2D(Resolution.Resolution.X, Resolution.Resolution.Y);
+			}
+		}
 		else
 		{
 			UMoviePipelineOutputSetting* OutputSetting = Cast<UMoviePipelineOutputSetting>(InJob->GetConfiguration()->FindOrAddSettingByClass(UMoviePipelineOutputSetting::StaticClass()));
@@ -262,13 +272,13 @@ void UMoviePipelinePIEExecutor::OnPIEStartupFinished(bool)
 
 	// Ensure the pipeline class we're using is compatible with the job's configuration type
 	const UMoviePipelineExecutorJob* CurrentJob = Queue->GetJobs()[CurrentPipelineIndex];
-	const bool bIsUsingGraph = CurrentJob->IsUsingGraphConfiguration();
-	if (bIsUsingGraph && !PipelineClass->IsChildOf(UMovieGraphPipeline::StaticClass()))
+	const bool bNeedsGraphPipeline = CurrentJob->IsUsingGraphConfiguration() || CurrentJob->IsUsingBasicConfiguration();
+	if (bNeedsGraphPipeline && !PipelineClass->IsChildOf(UMovieGraphPipeline::StaticClass()))
 	{
 		UE_LOG(LogMovieRenderPipeline, Log, TEXT("Using the default MovieGraphPipeline; the pipeline specified in Project Settings (%s) is not compatible with graphs."), *PipelineClass->GetName());
 		PipelineClass = UMovieGraphPipeline::StaticClass();
 	}
-	else if (!bIsUsingGraph && !PipelineClass->IsChildOf(UMoviePipeline::StaticClass()))
+	else if (!bNeedsGraphPipeline && !PipelineClass->IsChildOf(UMoviePipeline::StaticClass()))
 	{
 		UE_LOG(LogMovieRenderPipeline, Log, TEXT("Using the default MoviePipeline; the pipeline specified in Project Settings (%s) is not compatible with normal (non-graph) presets."), *PipelineClass->GetName());
 		PipelineClass = UMoviePipeline::StaticClass();
diff --git a/Engine/Plugins/MovieScene/MovieRenderPipeline/Source/MovieRenderPipelineEditor/Public/MoviePipelineEditorBlueprintLibrary.h b/Engine/Plugins/MovieScene/MovieRenderPipeline/Source/MovieRenderPipelineEditor/Public/MoviePipelineEditorBlueprintLibrary.h
--- a/Engine/Plugins/MovieScene/MovieRenderPipeline/Source/MovieRenderPipelineEditor/Public/MoviePipelineEditorBlueprintLibrary.h
+++ b/Engine/Plugins/MovieScene/MovieRenderPipeline/Source/MovieRenderPipelineEditor/Public/MoviePipelineEditorBlueprintLibrary.h
@@ -14,6 +14,7 @@
 #define UE_API MOVIERENDERPIPELINEEDITOR_API
 
 // Forward Declare
+class UMovieGraphConfig;
 class UMoviePipelinePrimaryConfig;
 class UMoviePipelineExecutorJob;
 class ULevelSequence;
@@ -58,6 +59,18 @@ public:
 	/** Resolves as much of the output directory for this job into a usable directory path as possible. Cannot resolve anything that relies on shot name, frame numbers, etc. */
 	UFUNCTION(BlueprintCallable, Category = "Movie Render Pipeline")
 	static UE_API FString ResolveOutputDirectoryFromJob(UMoviePipelineExecutorJob* InJob);
+
+	/**
+	 * Generates a UMovieGraphConfig from the job's Basic config.
+	 * Prompts the user to choose a save location via a Save Asset dialog, saves the asset to disk, then switches the
+	 * job to Graph configuration mode using the newly saved asset.
+	 *
+	 * @param InJob  The job to convert. Must be using Basic configuration mode.
+	 * @return true if the graph was saved and the job was switched; false if the job is not in Basic mode,
+	 *         graph generation or asset creation failed, the user canceled, or the save failed.
+	 */
+	UFUNCTION(BlueprintCallable, Category = "Movie Render Pipeline")
+	static UE_API bool SaveBasicConfigAsGraphConfig(UMoviePipelineExecutorJob* InJob);
 };
 
 #undef UE_API
