[MoviePipeline]: Add per-node render camera source system (FMovieGraphRenderCameraSource)

Introduces an extensible camera source interface that allows render pass nodes
to supply their own camera definitions, decoupling render passes from the
sequence camera system (primary camera + sidecar multi-camera data).

Renderer integration (UMovieGraphDefaultRenderer):
During shot setup, each render node instance is queried for a camera source.
If a source is present, camera count and indices come from the source's provider
space (0..N-1). Without a source, indices map to sequence sidecar camera data
(-1 for primary, 0..N-1 for sidecar). GetCameraInfo() falls back to sequence
cameras if a source cannot resolve a given index.

#Synced-CL 0

diff --git a/Engine/Plugins/MovieScene/MovieRenderPipeline/Source/MovieRenderPipelineCore/Private/Graph/MovieGraphDefaultRenderer.cpp b/Engine/Plugins/MovieScene/MovieRenderPipeline/Source/MovieRenderPipelineCore/Private/Graph/MovieGraphDefaultRenderer.cpp
--- a/Engine/Plugins/MovieScene/MovieRenderPipeline/Source/MovieRenderPipelineCore/Private/Graph/MovieGraphDefaultRenderer.cpp
+++ b/Engine/Plugins/MovieScene/MovieRenderPipeline/Source/MovieRenderPipelineCore/Private/Graph/MovieGraphDefaultRenderer.cpp
@@ -123,6 +123,10 @@ void UMovieGraphDefaultRenderer::SetupRenderingPipelineForShot(UMoviePipelineExe
 		}
 	}
 
+	// Tracks sources that have already had OnBeginFrame() called this frame.
+	// A source may be shared across multiple nodes, so this prevents duplicate calls.
+	TArray<TSharedPtr<FMovieGraphRenderCameraSource>> ProcessedCameraSources;
+
 	//UE_LOGF(LogMovieRenderPipeline, Warning, "Found: %d Render Passes:", OutputPasses.Num());
 	int32 TotalLayerCount = 0;
 	int32 CameraCount = 1;
@@ -141,19 +145,34 @@ void UMovieGraphDefaultRenderer::SetupRenderingPipelineForShot(UMoviePipelineExe
 			// Generate setup data for every render node / camera combination
 			for (const TWeakObjectPtr<UMovieGraphRenderPassNode>& RenderNode : BranchRenderer.Value)
 			{
-				for (int32 CameraIndex = 0; CameraIndex < CameraIndexes.Num(); CameraIndex++)
+				// Get or create the custom camera source for this node; null if the node uses default sequence cameras.
+				const TSharedPtr<FMovieGraphRenderCameraSource> RenderCameraSource = RenderNode->GetOrCreateRenderCameraSource(GetOwningGraph());
+				if (RenderCameraSource.IsValid() && !ProcessedCameraSources.Contains(RenderCameraSource))
+				{
+					// Prepare the source for the new frame; must be called before querying camera data.
+					RenderCameraSource->OnBeginFrame(GetOwningGraph(), &TimeStepData);
+					ProcessedCameraSources.Add(RenderCameraSource);
+				}
+
+				// Custom source drives camera count; otherwise fall back to the sequence camera index list.
+				const int32 NumCameras = RenderCameraSource ? RenderCameraSource->GetNumCameras() : CameraIndexes.Num();
+
+				// One layer entry per camera — each becomes an independent render pass.
+				for (int32 CameraIndex = 0; CameraIndex < NumCameras; CameraIndex++)
 				{
 					FMovieGraphRenderPassLayerData& LayerData = SetupData.Layers.AddDefaulted_GetRef();
 					LayerData.BranchName = BranchRenderer.Key.RootBranchName;
 					LayerData.LayerName = BranchRenderer.Key.LayerName;
 					LayerData.LayerIndex = *IdentifierToLayerIndex.Find(BranchRenderer.Key);
 					LayerData.RenderPassNode = RenderNode;
-					// Result could be -1 or 0...n
-					LayerData.CameraIndex = CameraIndexes[CameraIndex]; 
-					LayerData.NumCameras = CameraIndexes.Num();
+
+					// With a custom camera source the loop index is used directly (provider space, 0..N-1).
+					// Without one, CameraIndexes translates the loop index to -1 (primary) or 0..N-1 (sidecar).
+					LayerData.CameraIndex = RenderCameraSource ? CameraIndex : CameraIndexes[CameraIndex];
+					LayerData.NumCameras = NumCameras;
 				
 					// We provide the name here so that the setup functions of the renderers can create RenderLayerIdentifiers.
-					UE::MovieGraph::DefaultRenderer::FCameraInfo CameraInfo = GetCameraInfo(LayerData.CameraIndex);
+					UE::MovieGraph::DefaultRenderer::FCameraInfo CameraInfo = GetCameraInfo(LayerData.CameraIndex, RenderCameraSource.Get());
 					LayerData.CameraName = CameraInfo.CameraName;
 				}
 			}
@@ -381,6 +400,19 @@ void UMovieGraphDefaultRenderer::Render(const FMovieGraphTimeStepData& InTimeSte
 		}
 	}
 
+	// Collect all camera sources
+	TArray<TSharedRef<FMovieGraphRenderCameraSource>> AllCameraSources;
+	for (const TObjectPtr<UMovieGraphRenderPassNode>& RenderPass : RenderPassesInUse)
+	{
+		RenderPass->GatherRenderCameraSources(AllCameraSources);
+	}
+
+	// Notify once per frame 
+	for (const TSharedRef<FMovieGraphRenderCameraSource>& CameraSourceIt : AllCameraSources)
+	{
+		CameraSourceIt->OnBeginFrame(GetOwningGraph(), &InTimeStepData);
+	}
+
 	for (const TObjectPtr<UMovieGraphRenderPassNode>& RenderPass : RenderPassesInUse)
 	{
 		// Pass in a copy of the traversal context so the renderer can decide what to do with it.
@@ -482,6 +514,31 @@ void UMovieGraphDefaultRenderer::FlushAsyncEngineSystems(const TObjectPtr<UMovie
 
 void UMovieGraphDefaultRenderer::GetCameraLocationsForFrame(TArray<FVector>& OutLocations, UMoviePipelineExecutorShot* InShot, bool bIncludeSidecar) const
 {
+	// Collect custom camera sources from all active render pass nodes.
+	TArray<TSharedRef<FMovieGraphRenderCameraSource>> AllCameraSources;
+	for (const TObjectPtr<UMovieGraphRenderPassNode>& NodeIt : RenderPassesInUse)
+	{
+		if (NodeIt)
+		{
+			NodeIt->GatherRenderCameraSources(AllCameraSources);
+		}
+	}
+
+	// Add a location for each camera provided by custom sources.
+	for (const TSharedRef<FMovieGraphRenderCameraSource>& CameraSourceIt : AllCameraSources)
+	{
+		const int32 NumCameras = CameraSourceIt->GetNumCameras();
+		for (int32 CameraIndex = 0; CameraIndex < NumCameras; CameraIndex++)
+		{
+			FMinimalViewInfo MinimalViewInfo;
+			if (CameraSourceIt->GetCameraViewInfo(CameraIndex, MinimalViewInfo))
+			{
+				OutLocations.AddUnique(MinimalViewInfo.Location);
+			}
+		}
+	}
+
+	// Always include default sequence cameras (active game camera and optionally sidecar cameras).
 	TArray<UE::MovieGraph::FMinimalCameraInfo> MinimalCameraInfos = GetOwningGraph()->GetDataSourceInstance()->GetCameraInformation(InShot, bIncludeSidecar);
 	for (const UE::MovieGraph::FMinimalCameraInfo& MinimalCameraInfo : MinimalCameraInfos)
 	{
@@ -497,7 +554,7 @@ void UMovieGraphDefaultRenderer::AddOutstandingRenderTask_AnyThread(UE::Tasks::F
 	OutstandingTasks.Add(MoveTemp(InTask));
 }
 
-UE::MovieGraph::DefaultRenderer::FCameraInfo UMovieGraphDefaultRenderer::GetCameraInfo(const int32 InCameraIndex) const
+UE::MovieGraph::DefaultRenderer::FCameraInfo UMovieGraphDefaultRenderer::GetCameraInfoImpl(const int32 InCameraIndex) const
 {
 	UE::MovieGraph::DefaultRenderer::FCameraInfo CameraInfo;
 	
@@ -533,7 +590,49 @@ UE::MovieGraph::DefaultRenderer::FCameraInfo UMovieGraphDefaultRenderer::GetCame
 	return CameraInfo;
 }
 
-float UMovieGraphDefaultRenderer::GetCameraOverscan(int32 InCameraIndex)
+UE::MovieGraph::DefaultRenderer::FCameraInfo UMovieGraphDefaultRenderer::GetCameraInfo(	
+	int32 InCameraIndex,
+	const FMovieGraphRenderCameraSource* InRenderCameraSource) const
+{
+	UE::MovieGraph::DefaultRenderer::FCameraInfo CameraInfo;
+
+	// When a custom camera source is active, query it directly instead of the default sequence cameras.
+	if (InRenderCameraSource && InCameraIndex >= 0)
+	{
+		FMinimalViewInfo ViewInfo;
+		if (InRenderCameraSource->GetCameraViewInfo(InCameraIndex, ViewInfo))
+		{
+			CameraInfo.ViewInfo = ViewInfo;
+
+			// GetCameraViewActor() and GetCameraName() are optional — the source may defer
+			// either to the default sequence camera by returning false.
+			const bool bHasCameraViewActor = InRenderCameraSource->GetCameraViewActor(InCameraIndex, CameraInfo.ViewActor);
+			const bool bHasCameraName = InRenderCameraSource->GetCameraName(InCameraIndex, CameraInfo.CameraName);
+			if (!bHasCameraViewActor || !bHasCameraName)
+			{
+				const UE::MovieGraph::DefaultRenderer::FCameraInfo DefaultCameraInfo = GetCameraInfoImpl(INDEX_NONE);
+				if (!bHasCameraViewActor)
+				{
+					CameraInfo.ViewActor = DefaultCameraInfo.ViewActor;
+				}
+				if (!bHasCameraName)
+				{
+					CameraInfo.CameraName = DefaultCameraInfo.CameraName;
+				}
+			}
+
+			return CameraInfo;
+		}
+
+		// Source exists but couldn't resolve this index — fall back to the primary sequence camera.
+		// INDEX_NONE is the sentinel for "primary camera" in the default camera path below.
+		return GetCameraInfoImpl(INDEX_NONE);
+	}
+
+	return GetCameraInfoImpl(InCameraIndex);
+}
+
+float UMovieGraphDefaultRenderer::GetCameraOverscanImpl(int32 InCameraIndex)
 {
 	if (CameraOverscanCache.Contains(InCameraIndex))
 	{
@@ -545,6 +644,30 @@ float UMovieGraphDefaultRenderer::GetCameraOverscan(int32 InCameraIndex)
 	return CameraOverscan;
 }
 
+float UMovieGraphDefaultRenderer::GetCameraOverscan(
+	int32 InCameraIndex,
+	const FMovieGraphRenderCameraSource* InRenderCameraSource)
+{
+	// When CameraSource overrides default cameras, caches can't be used. The cache is keyed only by
+	// InCameraIndex, but a CameraSource supplies a completely different set of cameras than the default
+	// sequence cameras. The same index can therefore resolve to a different view with a different
+	// overscan, making any cached value incorrect.
+	// Note: INDEX_NONE reserved for default camera.
+	if (InRenderCameraSource)
+	{
+		float RenderCameraSourceOverscan = 0.f;
+		if (InRenderCameraSource->GetCameraOverscan(InCameraIndex, RenderCameraSourceOverscan))
+		{
+			return RenderCameraSourceOverscan;
+		}
+
+		// fallback to default camera
+		return GetCameraOverscanImpl(INDEX_NONE);
+	}
+
+	return GetCameraOverscanImpl(InCameraIndex);
+}
+
 void UMovieGraphDefaultRenderer::WarnAboutAnimatedOverscan(float InInitialOverscan)
 {
 	if (!bHasWarnedAboutAnimatedOverscan)
diff --git a/Engine/Plugins/MovieScene/MovieRenderPipeline/Source/MovieRenderPipelineCore/Private/Graph/MovieGraphSequenceDataSource.cpp b/Engine/Plugins/MovieScene/MovieRenderPipeline/Source/MovieRenderPipelineCore/Private/Graph/MovieGraphSequenceDataSource.cpp
--- a/Engine/Plugins/MovieScene/MovieRenderPipeline/Source/MovieRenderPipelineCore/Private/Graph/MovieGraphSequenceDataSource.cpp
+++ b/Engine/Plugins/MovieScene/MovieRenderPipeline/Source/MovieRenderPipelineCore/Private/Graph/MovieGraphSequenceDataSource.cpp
@@ -419,6 +419,16 @@ TArray<UE::MovieGraph::FMinimalCameraInfo> UMovieGraphSequenceDataSource::GetCam
 	return OutInfo;
 }
 
+UMovieSceneSequencePlayer* UMovieGraphSequenceDataSource::GetSequencePlayer() const
+{
+	if (LevelSequenceActor)
+	{
+		return LevelSequenceActor->GetSequencePlayer();
+	}
+
+	return nullptr;
+}
+
 namespace UE::MovieGraph
 {
 	FFrameTime FMovieGraphSequenceTimeController::OnRequestCurrentTime(const FQualifiedFrameTime& InCurrentTime, float InPlayRate)
diff --git a/Engine/Plugins/MovieScene/MovieRenderPipeline/Source/MovieRenderPipelineCore/Public/Graph/MovieGraphDefaultRenderer.h b/Engine/Plugins/MovieScene/MovieRenderPipeline/Source/MovieRenderPipelineCore/Public/Graph/MovieGraphDefaultRenderer.h
--- a/Engine/Plugins/MovieScene/MovieRenderPipeline/Source/MovieRenderPipelineCore/Public/Graph/MovieGraphDefaultRenderer.h
+++ b/Engine/Plugins/MovieScene/MovieRenderPipeline/Source/MovieRenderPipelineCore/Public/Graph/MovieGraphDefaultRenderer.h
@@ -16,6 +16,7 @@
 class UMovieGraphRenderPassNode;
 class UTextureRenderTarget2D;
 struct FMoviePipelineSurfaceQueue;
+struct FMovieGraphRenderCameraSource;
 namespace MoviePipeline { struct IMoviePipelineOverlappedAccumulator; }
 namespace UE::MovieGraph::DefaultRenderer { struct FSurfaceAccumulatorPool; }
 
@@ -314,16 +315,28 @@ public:
 	// ~UObject Interface
 
 	UE_API void AddOutstandingRenderTask_AnyThread(UE::Tasks::FTask InTask);
-	/** Fetches information for the given camera index. Should be "-1" when not using multi-camera rendering, or [0, n] based on the shot's sidecar camera data when using multi-camera rendering. */
-	UE_API UE::MovieGraph::DefaultRenderer::FCameraInfo GetCameraInfo(const int32 InCameraIndex) const;
+	
+	/**
+	 * Fetches information for the given camera index.
+	 * First, if InCameraSource is provided with a non-negative camera index, its information is used.
+	 * Otherwise, it is resolved against default sequence camera or the shot's sidecar camera data.
+	 * @param InCameraIndex  Pass INDEX_NONE when using default sequence camera.
+	 *                       Pass [0..N) for multi-camera rendering (shot's sidecar or custom source).
+	 * @param InCameraSource Optional camera source. If provided but the camera index cannot be
+	 *                       resolved, falls back to the default sequence camera.
+	 */
+	UE_API UE::MovieGraph::DefaultRenderer::FCameraInfo GetCameraInfo(const int32 InCameraIndex, const FMovieGraphRenderCameraSource* InCameraSource = nullptr) const;
 	void SetHasRenderedFirstViewThisFrame(bool bInValue) { bHasRenderedFirstViewThisFrame = bInValue; }
 	bool GetHasRenderedFirstViewThisFrame() const { return bHasRenderedFirstViewThisFrame; }
 
 	/**
-	 * Gets the overscan value for the specified camera. First checks to see if any cached overscan value exists and returning it, and if
-	 * no cached value was found, gets the camera's live overscan value and caches it.
+	 * Gets the overscan value for the specified camera.
+	 * @param InCameraIndex  Pass INDEX_NONE when using default sequence camera.
+	 *                       Pass [0..N) for multi-camera rendering (shot's sidecar or custom source).
+	 * @param InCameraSource Optional camera source. If provided but the camera index cannot be
+	 *                       resolved, falls back to the default sequence camera.
 	 */
-	UE_API float GetCameraOverscan(int32 InCameraIndex);
+	UE_API float GetCameraOverscan(int32 InCameraIndex, const FMovieGraphRenderCameraSource* InCameraSource = nullptr);
 
 	/** Outputs a warning message regarding animated overscan to the MRQ log if one has not already been output  */
 	UE_DEPRECATED(5.6, "WarnAboutAnimatedOverscan is no longer used; animated overscan is supported in 5.6")
@@ -360,6 +373,14 @@ protected:
 	UE_API void GetCameraLocationsForFrame(TArray<FVector>& OutLocations, UMoviePipelineExecutorShot* InShot, bool bIncludeSidecar = false) const;
 	UE_API void FlushAsyncEngineSystems(const TObjectPtr<UMovieGraphEvaluatedConfig>& InConfig) const;
 
+	/** Fetches information for the given camera index. Should be "-1" when not using multi-camera rendering, or [0, n] based on the shot's sidecar camera data when using multi-camera rendering. */
+	UE_API UE::MovieGraph::DefaultRenderer::FCameraInfo GetCameraInfoImpl(const int32 InCameraIndex) const;
+	/**
+	 * Gets the overscan value for the specified camera. First checks to see if any cached overscan value exists and returning it, and if
+	 * no cached value was found, gets the camera's live overscan value and caches it.
+	 */
+	UE_API float GetCameraOverscanImpl(int32 InCameraIndex);
+
 protected:
 	/** A pointer to the CDOs of the Render Pass nodes that are valid for the current shot render. */
 	UPROPERTY(Transient)
diff --git a/Engine/Plugins/MovieScene/MovieRenderPipeline/Source/MovieRenderPipelineCore/Public/Graph/MovieGraphRenderCameraSource.h b/Engine/Plugins/MovieScene/MovieRenderPipeline/Source/MovieRenderPipelineCore/Public/Graph/MovieGraphRenderCameraSource.h
new file mode 100644
--- /dev/null
+++ b/Engine/Plugins/MovieScene/MovieRenderPipeline/Source/MovieRenderPipelineCore/Public/Graph/MovieGraphRenderCameraSource.h
@@ -0,0 +1,104 @@
+// Copyright Epic Games, Inc. All Rights Reserved.
+
+#pragma once
+#include "CoreMinimal.h"
+#include "Camera/CameraTypes.h"
+#include "SceneView.h"
+
+class FReferenceCollector;
+class AActor;
+class UMovieGraphPipeline;
+struct FMovieGraphTimeStepData;
+
+#define UE_API MOVIERENDERPIPELINECORE_API
+
+/**
+ * Base class for a custom camera source created by a live UMovieGraphRenderPassNode instance
+ * (not the class default object). Supplies a set of cameras that replaces the cameras from
+ * the data source for the render pass node that created it.
+ */
+struct FMovieGraphRenderCameraSource
+{
+	virtual ~FMovieGraphRenderCameraSource() = default;
+
+	/**
+	 * Adds any UObject references held by this source to the collector for GC tracking.
+	 * Subclasses that hold UObject pointers must override this to prevent those objects
+	 * from being garbage collected while the source is alive.
+	 *
+	 * @param Collector (in) The reference collector provided by the GC system.
+	 */
+	UE_API virtual void AddReferencedObjects(FReferenceCollector& Collector) {}
+
+	/**
+	 * Called once per frame to prepare the source for a new frame. Must be called before querying camera data.
+	 * 
+	 * @param InMoviePipeline (in) The pipeline currently being evaluated.
+	 * @param InTimeData      (in) Time step data for the current frame (frame time, delta time, etc.).
+	 */
+	UE_API virtual void OnBeginFrame(const UMovieGraphPipeline* InMoviePipeline, const FMovieGraphTimeStepData* InTimeData) { }
+
+	/** Returns the number of cameras this source contributes for the current shot. */
+	UE_API virtual int32 GetNumCameras() const { return 0; }
+
+	/**
+	 * Returns the unique camera name for the given camera index.
+	 * Used by the renderer for the {camera_name} output token.
+	 *
+	 * @param InCameraIndex (in)  Index in range [0, NumCameras).
+	 * @param OutCameraName (out) Receives the camera name for the requested index.
+	 * @return              True if resolved; false to fall back to the shot's camera name.
+	 */
+	UE_API virtual bool GetCameraName(const int32 InCameraIndex, FString& OutCameraName) const { return false; }
+
+	/**
+	 * Returns the view actor for the specified camera index.
+	 * Returning true with OutCameraActor = nullptr is valid and suppresses the default fallback.
+	 *
+	 * @param InCameraIndex  (in)  Index in range [0, NumCameras).
+	 * @param OutCameraActor (out) Receives the actor used as the view origin, or nullptr.
+	 * @return               True if the source handled this query; false to fall back to the player controller's view target.
+	 */
+	UE_API virtual bool GetCameraViewActor(const int32 InCameraIndex, AActor*& OutCameraActor) const { return false; }
+
+	/**
+	 * Gets the overscan value for the specified camera.
+	 *
+	 * @param InCameraIndex     (in)  Index in range [0, NumCameras).
+	 * @param OutCameraOverscan (out) Receives the overscan value for the camera.
+	 * @return                  True if resolved; false to fall back to the default sequence camera overscan.
+	 */
+	UE_API virtual bool GetCameraOverscan(const int32 InCameraIndex, float& OutCameraOverscan) const { return false; }
+
+	/**
+	 * Returns the view info for the specified camera index.
+	 *
+	 * @param InCameraIndex (in)  Index in range [0, NumCameras).
+	 * @param OutViewInfo   (out) Receives the view transform, FOV, and other camera parameters.
+	 * @return              True if resolved; false to fall back to the primary sequence camera.
+	 */
+	UE_API virtual bool GetCameraViewInfo(const int32 InCameraIndex, FMinimalViewInfo& OutViewInfo) const { return false; }
+
+	/**
+	 * Applies camera-specific adjustments to the projection data after it has been computed.
+	 *
+	 * @param InCameraIndex       (in)     Index in range [0, NumCameras).
+	 * @param InOutProjectionData (in,out) Projection data to adjust; already populated by the time this is called.
+	 */
+	UE_API virtual void SetupCameraViewProjectionData(const int32 InCameraIndex, FSceneViewProjectionData& InOutProjectionData) const {}
+
+	/**
+	 * Returns the overscanned resolution and inner crop rectangle for the specified camera.
+	 *
+	 * @param InCameraIndex            (in)  Index in range [0, NumCameras).
+	 * @param OutOverscannedResolution (out) Receives the overscanned resolution.
+	 * @param OutOverscanCropRectangle (out) Receives the inner crop rectangle within the overscanned resolution.
+	 * @return                         True if resolved; false to fall back to the default resolution calculation.
+	 */
+	UE_API virtual bool GetCameraOverscannedResolution(
+		const int32 InCameraIndex,
+		FIntPoint& OutOverscannedResolution,
+		FIntRect& OutOverscanCropRectangle) const { return false; }
+};
+
+#undef UE_API
\ No newline at end of file
diff --git a/Engine/Plugins/MovieScene/MovieRenderPipeline/Source/MovieRenderPipelineCore/Public/Graph/MovieGraphSequenceDataSource.h b/Engine/Plugins/MovieScene/MovieRenderPipeline/Source/MovieRenderPipelineCore/Public/Graph/MovieGraphSequenceDataSource.h
--- a/Engine/Plugins/MovieScene/MovieRenderPipeline/Source/MovieRenderPipelineCore/Public/Graph/MovieGraphSequenceDataSource.h
+++ b/Engine/Plugins/MovieScene/MovieRenderPipeline/Source/MovieRenderPipelineCore/Public/Graph/MovieGraphSequenceDataSource.h
@@ -58,6 +58,11 @@ public:
 		const int32 InRightDeltaFrames, const bool bInPrepass) override;
 	UE_API virtual TArray<UE::MovieGraph::FMinimalCameraInfo> GetCameraInformation(UMoviePipelineExecutorShot* InShot, bool bIncludeSidecar) const override;
 	UE_API virtual void InitializeShot(const TObjectPtr<UMoviePipelineExecutorShot>& InShot, const FFrameTime& InEvalTime) override;
+
+public:
+	/** Returns the sequence player used to evaluate the level sequence during rendering. */
+	UE_API UMovieSceneSequencePlayer* GetSequencePlayer() const;
+
 protected:
 	UE_API void OverrideSequencePlaybackRangeFromGlobalOutputSettings(ULevelSequence* InSequence);
 	UE_API void CacheLevelSequenceData(ULevelSequence* InSequence);
diff --git a/Engine/Plugins/MovieScene/MovieRenderPipeline/Source/MovieRenderPipelineCore/Public/Graph/Nodes/MovieGraphRenderPassNode.h b/Engine/Plugins/MovieScene/MovieRenderPipeline/Source/MovieRenderPipelineCore/Public/Graph/Nodes/MovieGraphRenderPassNode.h
--- a/Engine/Plugins/MovieScene/MovieRenderPipeline/Source/MovieRenderPipelineCore/Public/Graph/Nodes/MovieGraphRenderPassNode.h
+++ b/Engine/Plugins/MovieScene/MovieRenderPipeline/Source/MovieRenderPipelineCore/Public/Graph/Nodes/MovieGraphRenderPassNode.h
@@ -4,6 +4,7 @@
 
 #include "Graph/MovieGraphNode.h"
 #include "Graph/MovieGraphDataTypes.h"
+#include "Graph/MovieGraphRenderCameraSource.h"
 #include "Styling/AppStyle.h"
 
 #include "MovieGraphRenderPassNode.generated.h"
@@ -42,6 +43,19 @@ public:
 	 */
 	FString GetRendererSubName() const { return GetRendererSubNameImpl(); }
 	
+	/**
+	 * Create a new (or gets an exists) camera source for this render node instance, or nullptr if this node uses
+	 * default sequence cameras. Called on non-CDO instances.
+	 */
+	TSharedPtr<FMovieGraphRenderCameraSource> GetOrCreateRenderCameraSource(const UMovieGraphPipeline* InPipeline) { return GetOrCreateRenderCameraSourceImpl(InPipeline); }
+
+	/**
+	 * Collects custom camera sources from all active pass instances into InOutSources.
+	 * Called on the CDO.
+	 * @param InOutSources (in, out) May be non-empty on input; unique sources are appended.
+	 */
+	void GatherRenderCameraSources(TArray<TSharedRef<FMovieGraphRenderCameraSource>>& InOutSources) const { GatherRenderCameraSourcesImpl(InOutSources); }
+
 	/** Called when this should set up for rendering a new shot. Called on the CDO. */
 	void Setup(const FMovieGraphRenderPassSetupData& InSetupData) { SetupImpl(InSetupData); }
 	
@@ -125,6 +139,8 @@ protected:
 	virtual void GatherOutputPassesImpl(UMovieGraphEvaluatedConfig* InConfig, TArray<FMovieGraphRenderDataIdentifier>& OutExpectedPasses) const {}
 	virtual int32 GetNumSceneViewsRenderedImpl() const { return 0; }
 	virtual TArray<FSoftClassPath> GetOutputTypeRestrictionsImpl() const { return {}; } 
+	virtual TSharedPtr<FMovieGraphRenderCameraSource> GetOrCreateRenderCameraSourceImpl(const UMovieGraphPipeline* InMovieGraphPipeline) { return nullptr; }
+	virtual void GatherRenderCameraSourcesImpl(TArray<TSharedRef<FMovieGraphRenderCameraSource>>& InOutSources) const {};
 };
 
 #undef UE_API
diff --git a/Engine/Plugins/MovieScene/MovieRenderPipeline/Source/MovieRenderPipelineRenderPasses/Private/Graph/Nodes/MovieGraphImagePassBaseNode.cpp b/Engine/Plugins/MovieScene/MovieRenderPipeline/Source/MovieRenderPipelineRenderPasses/Private/Graph/Nodes/MovieGraphImagePassBaseNode.cpp
--- a/Engine/Plugins/MovieScene/MovieRenderPipeline/Source/MovieRenderPipelineRenderPasses/Private/Graph/Nodes/MovieGraphImagePassBaseNode.cpp
+++ b/Engine/Plugins/MovieScene/MovieRenderPipeline/Source/MovieRenderPipelineRenderPasses/Private/Graph/Nodes/MovieGraphImagePassBaseNode.cpp
@@ -110,6 +110,22 @@ void UMovieGraphImagePassBaseNode::GatherOutputPassesImpl(UMovieGraphEvaluatedCo
 	}
 }
 
+void UMovieGraphImagePassBaseNode::GatherRenderCameraSourcesImpl(TArray<TSharedRef<FMovieGraphRenderCameraSource>>& InOutSources) const
+{
+	// Note: InOutSources may already contain sources from other nodes; AddUnique() avoids duplicates.
+	
+	for (const TUniquePtr<UE::MovieGraph::Rendering::FMovieGraphImagePassBase>& InstanceIt : CurrentInstances)
+	{
+		const TSharedPtr<FMovieGraphRenderCameraSource> CameraSource = InstanceIt.IsValid()
+			? InstanceIt->GetRenderCameraSource() : nullptr;
+
+		if (CameraSource)
+		{
+			InOutSources.AddUnique(CameraSource.ToSharedRef());
+		}
+	}
+}
+
 EViewModeIndex UMovieGraphImagePassBaseNode::GetViewModeIndex() const
 {
 	return VMI_Lit;
diff --git a/Engine/Plugins/MovieScene/MovieRenderPipeline/Source/MovieRenderPipelineRenderPasses/Private/Graph/Renderers/MovieGraphDeferredPanoramicPass.cpp b/Engine/Plugins/MovieScene/MovieRenderPipeline/Source/MovieRenderPipelineRenderPasses/Private/Graph/Renderers/MovieGraphDeferredPanoramicPass.cpp
--- a/Engine/Plugins/MovieScene/MovieRenderPipeline/Source/MovieRenderPipelineRenderPasses/Private/Graph/Renderers/MovieGraphDeferredPanoramicPass.cpp
+++ b/Engine/Plugins/MovieScene/MovieRenderPipeline/Source/MovieRenderPipelineRenderPasses/Private/Graph/Renderers/MovieGraphDeferredPanoramicPass.cpp
@@ -285,7 +285,7 @@ void FMovieGraphDeferredPanoramicPass::Setup(TWeakObjectPtr<UMovieGraphDefaultRe
 
 		PaneViewStates.SetNum(NumPanoramicPanes);
 
-		const UE::MovieGraph::DefaultRenderer::FCameraInfo& CameraInfo = InRenderer->GetCameraInfo(InLayer.CameraIndex);
+		const UE::MovieGraph::DefaultRenderer::FCameraInfo& CameraInfo = GetRenderCameraInfo();
 
 		// ToDo: This doesn't take into account blended post-process values from the world, but we don't have a way to do the blending without having a FSceneView which doesn't exist until render time.
 		const bool bUsesAutoExposure = CameraInfo.ViewInfo.PostProcessSettings.bOverride_AutoExposureMethod && CameraInfo.ViewInfo.PostProcessSettings.AutoExposureMethod != EAutoExposureMethod::AEM_Manual;
@@ -430,6 +430,11 @@ FName FMovieGraphDeferredPanoramicPass::GetBranchName() const
 	return LayerData.BranchName;
 }
 
+int32 FMovieGraphDeferredPanoramicPass::GetRenderCameraIndex() const
+{
+	return LayerData.CameraIndex;
+}
+
 void FMovieGraphDeferredPanoramicPass::Render(const FMovieGraphTraversalContext& InFrameTraversalContext, const FMovieGraphTimeStepData& InTimeData)
 {
 	FMovieGraphImagePassBase::Render(InFrameTraversalContext, InTimeData);
@@ -499,7 +504,7 @@ void FMovieGraphDeferredPanoramicPass::Render(const FMovieGraphTraversalContext&
 		{
 			for (int32 SpatialIndex = 0; SpatialIndex < NumSpatialSamples; SpatialIndex++)
 			{
-				UE::MovieGraph::DefaultRenderer::FCameraInfo CameraInfo = GetRenderer()->GetCameraInfo(LayerData.CameraIndex);
+				UE::MovieGraph::DefaultRenderer::FCameraInfo CameraInfo = GetRenderCameraInfo();
 
 				UE::MoviePipeline::FPanoramicPane Data;
 				
diff --git a/Engine/Plugins/MovieScene/MovieRenderPipeline/Source/MovieRenderPipelineRenderPasses/Private/Graph/Renderers/MovieGraphDeferredPass.cpp b/Engine/Plugins/MovieScene/MovieRenderPipeline/Source/MovieRenderPipelineRenderPasses/Private/Graph/Renderers/MovieGraphDeferredPass.cpp
--- a/Engine/Plugins/MovieScene/MovieRenderPipeline/Source/MovieRenderPipelineRenderPasses/Private/Graph/Renderers/MovieGraphDeferredPass.cpp
+++ b/Engine/Plugins/MovieScene/MovieRenderPipeline/Source/MovieRenderPipelineRenderPasses/Private/Graph/Renderers/MovieGraphDeferredPass.cpp
@@ -78,6 +78,11 @@ namespace UE::MovieGraph::Rendering
 		}
 	}
 
+int32 FMovieGraphDeferredPass::GetRenderCameraIndex() const
+{
+	return LayerData.CameraIndex;
+}
+
 FSceneViewStateInterface* FMovieGraphDeferredPass::GetSceneViewState(UMovieGraphImagePassBaseNode* ParentNodeThisFrame, int32_t TileX, int32_t TileY)
 {
 	check(ParentNodeThisFrame);
@@ -103,12 +108,12 @@ FMovieGraphDeferredPass::FResolutionAndCameraInfo FMovieGraphDeferredPass::GetRe
 	constexpr bool bIncludeCDOs = false;
 	const UMovieGraphCameraSettingNode* CameraNode = InEvaluatedConfig->GetSettingForBranch<UMovieGraphCameraSettingNode>(LayerData.BranchName, bIncludeCDOs);
 
-	Info.CameraInfo = GetRenderer()->GetCameraInfo(LayerData.CameraIndex);
+	Info.CameraInfo = GetRenderCameraInfo();
 
 	// Get the overscan fraction and if all cameras should be rendered
 	{
 		// By default, use the overscan fraction specified directly on the camera actor
-		Info.OverscanFraction = GetRenderer()->GetCameraOverscan(LayerData.CameraIndex);
+		Info.OverscanFraction = GetRenderCameraOverscan();
 
 		// The Camera Settings node in the graph may specify an overscan override
 		if (CameraNode)
@@ -134,6 +139,16 @@ FMovieGraphDeferredPass::FResolutionAndCameraInfo FMovieGraphDeferredPass::GetRe
 	// Specifies a crop within the AccumulatorResolution for us to take center-out crops later where needed.
 	Info.AccumulatorResolutionCropRect = UMovieGraphBlueprintLibrary::GetOverscanCropRectangle(InEvaluatedConfig, Info.OverscanFraction, Info.CameraAspectRatio);
 
+	// If a custom camera source is active, let it override the accumulator resolution and crop rect
+	// computed above with per-camera values.
+	const TSharedPtr<const FMovieGraphRenderCameraSource> CameraSource = GetRenderCameraSource();
+	if (CameraSource && CameraSource->GetCameraOverscannedResolution(GetRenderCameraIndex(), Info.AccumulatorResolution, Info.AccumulatorResolutionCropRect))
+	{
+		// A custom camera source provides multiple cameras, so bRenderAllCameras must be
+		// set to override the default single-camera path and render all of them.
+		Info.bRenderAllCameras = true;
+	}
+
 	// This is what the actual renders would be at, without overlap.
 	const FIntPoint TileCount = ParentNodeThisFrame->GetTileCount();
 	Info.BackbufferResolution = FIntPoint(
@@ -200,7 +215,7 @@ void FMovieGraphDeferredPass::Setup(TWeakObjectPtr<UMovieGraphDefaultRenderer> I
 		}
 
 		// ToDo: This doesn't take into account blended post-process values from the world, but we don't have a way to do the blending without having a FSceneView which doesn't exist until render time.
-		const UE::MovieGraph::DefaultRenderer::FCameraInfo& CameraInfo = InRenderer->GetCameraInfo(InLayer.CameraIndex);
+		const UE::MovieGraph::DefaultRenderer::FCameraInfo& CameraInfo = GetRenderCameraInfo();
 		const bool bUsesAutoExposure = CameraInfo.ViewInfo.PostProcessSettings.bOverride_AutoExposureMethod && CameraInfo.ViewInfo.PostProcessSettings.AutoExposureMethod != EAutoExposureMethod::AEM_Manual;
 
 		const bool bHighRes = ParentNode->GetEnableHighResolutionTiling();
@@ -291,6 +306,8 @@ void FMovieGraphDeferredPass::GatherOutputPasses(UMovieGraphEvaluatedConfig* InC
 
 void FMovieGraphDeferredPass::AddReferencedObjects(FReferenceCollector& Collector)
 {
+	FMovieGraphImagePassBase::AddReferencedObjects(Collector);
+
 	PRAGMA_DISABLE_DEPRECATION_WARNINGS
 	FSceneViewStateInterface* SceneViewRef = SceneViewState.GetReference();
 	if (SceneViewRef)
@@ -491,7 +508,7 @@ void FMovieGraphDeferredPass::Render(const FMovieGraphTraversalContext& InFrameT
 				}
 
 				// These are the parameters of our camera 
-				UE::MovieGraph::DefaultRenderer::FCameraInfo CameraInfo = GetRenderer()->GetCameraInfo(LayerData.CameraIndex);
+				UE::MovieGraph::DefaultRenderer::FCameraInfo CameraInfo = GetRenderCameraInfo();
 
 				FIntPoint TileSize = ResolutionAndCameraInfo.TileSize;
 				FIntPoint OverlappedPad = FIntPoint(FMath::CeilToInt(TileSize.X * TileOverlapFraction), FMath::CeilToInt(TileSize.Y * TileOverlapFraction));
@@ -590,7 +607,10 @@ void FMovieGraphDeferredPass::Render(const FMovieGraphTraversalContext& InFrameT
 				);
 				
 				// ToDo: If an actor had multiple cine camera components this could be wrong, but Sequencer also can't disambiguate between multiple camera components on a single actor.
-				UE::MoviePipeline::GetMetadataFromCineCamera(CameraInfo.ViewActor->GetComponentByClass<UCineCameraComponent>(), MetadataPrefix, SampleState.AdditionalFileMetadata);
+				UE::MoviePipeline::GetMetadataFromCineCamera(!CameraInfo.ViewActor ? nullptr : CameraInfo.ViewActor->GetComponentByClass<UCineCameraComponent>(), MetadataPrefix, SampleState.AdditionalFileMetadata);
+
+				// Additional metadata from render pass
+				GetMetadataFromRenderPass(SampleState.AdditionalFileMetadata);
 
 				// If this was just to contribute to the history buffer, no need to go any further. Never discard if we're writing individual samples, though.
 				bool bRenderPass = (!InTimeData.bDiscardOutput && !ShouldDiscardOutput(ViewFamily, CameraInfo)) || SampleState.bWriteSampleToDisk;
diff --git a/Engine/Plugins/MovieScene/MovieRenderPipeline/Source/MovieRenderPipelineRenderPasses/Private/Graph/Renderers/MovieGraphImagePassBase.cpp b/Engine/Plugins/MovieScene/MovieRenderPipeline/Source/MovieRenderPipelineRenderPasses/Private/Graph/Renderers/MovieGraphImagePassBase.cpp
--- a/Engine/Plugins/MovieScene/MovieRenderPipeline/Source/MovieRenderPipelineRenderPasses/Private/Graph/Renderers/MovieGraphImagePassBase.cpp
+++ b/Engine/Plugins/MovieScene/MovieRenderPipeline/Source/MovieRenderPipelineRenderPasses/Private/Graph/Renderers/MovieGraphImagePassBase.cpp
@@ -5,6 +5,7 @@
 #include "Graph/MovieGraphDataTypes.h"
 #include "Graph/MovieGraphPipeline.h"
 #include "Graph/Nodes/MovieGraphApplyViewportLookNode.h"
+#include "Graph/Nodes/MovieGraphImagePassBaseNode.h"
 #include "SceneView.h"
 #include "MoviePipelineUtils.h"
 #include "LegacyScreenPercentageDriver.h"
@@ -32,6 +33,49 @@ void FMovieGraphImagePassBase::Setup(TWeakObjectPtr<UMovieGraphDefaultRenderer>
 {
 	// This is a pointer to the UMovieGraphPipeline's renderer which is valid throughout the entire render.
 	WeakGraphRenderer = InRenderer;
+
+	// Fetch (or lazily create) the camera source from InLayer.RenderPassNode.
+	if (ensureMsgf(InLayer.RenderPassNode.IsValid(), TEXT("FMovieGraphImagePassBase::Setup called with an invalid RenderPassNode.")))
+	{
+		RenderCameraSource = InLayer.RenderPassNode.Get()->GetOrCreateRenderCameraSource(InRenderer->GetOwningGraph());
+	}
+}
+
+void FMovieGraphImagePassBase::Teardown()
+{
+	// Release our shared reference to the camera source. The source stays alive until every pass
+	// instance for this node has torn down and dropped its reference.
+	RenderCameraSource.Reset();
+}
+
+void FMovieGraphImagePassBase::AddReferencedObjects(FReferenceCollector& Collector)
+{
+	if (RenderCameraSource.IsValid())
+	{
+		RenderCameraSource->AddReferencedObjects(Collector);
+	}
+}
+
+UE::MovieGraph::DefaultRenderer::FCameraInfo FMovieGraphImagePassBase::GetRenderCameraInfo() const
+{
+	const TObjectPtr<UMovieGraphDefaultRenderer> GraphRenderer = GetRenderer().Get();
+	if (!GraphRenderer)
+	{
+		return UE::MovieGraph::DefaultRenderer::FCameraInfo();
+	}
+
+	return GraphRenderer->GetCameraInfo(GetRenderCameraIndex(), GetRenderCameraSource().Get());
+}
+
+float FMovieGraphImagePassBase::GetRenderCameraOverscan() const
+{
+	const TObjectPtr<UMovieGraphDefaultRenderer> GraphRenderer = GetRenderer().Get();
+	if (!GraphRenderer)
+	{
+		return 0.f;
+	}
+
+	return GraphRenderer->GetCameraOverscan(GetRenderCameraIndex(), GetRenderCameraSource().Get());
 }
 
 FSceneViewInitOptions FMovieGraphImagePassBase::CreateViewInitOptions(const UE::MovieGraph::DefaultRenderer::FCameraInfo& InCameraInfo, FSceneViewFamilyContext* InViewFamily, FSceneViewStateReference& InViewStateRef) const
@@ -406,6 +450,12 @@ void FMovieGraphImagePassBase::CalculateProjectionMatrix(UE::MovieGraph::Default
 
 	// This function updates data in both the FMinimalViewInfo and the ProjectionData
 	FMinimalViewInfo::CalculateProjectionMatrixGivenViewRectangle(InOutCameraInfo.ViewInfo, AspectRatioAxisConstraint, ViewExtents, InOutProjectionData);
+
+	// Give the camera source a chance to apply final adjustments to the computed projection data.
+	if (const TSharedPtr<const FMovieGraphRenderCameraSource> CameraSource = GetRenderCameraSource())
+	{
+		CameraSource->SetupCameraViewProjectionData(GetRenderCameraIndex(), InOutProjectionData);
+	}
 }
 
 FVector4f FMovieGraphImagePassBase::CalculatePrinciplePointOffsetForTiling(const UE::MovieGraph::DefaultRenderer::FMovieGraphTilingParams& InTilingParams) const 
diff --git a/Engine/Plugins/MovieScene/MovieRenderPipeline/Source/MovieRenderPipelineRenderPasses/Public/Graph/Nodes/MovieGraphImagePassBaseNode.h b/Engine/Plugins/MovieScene/MovieRenderPipeline/Source/MovieRenderPipelineRenderPasses/Public/Graph/Nodes/MovieGraphImagePassBaseNode.h
--- a/Engine/Plugins/MovieScene/MovieRenderPipeline/Source/MovieRenderPipelineRenderPasses/Public/Graph/Nodes/MovieGraphImagePassBaseNode.h
+++ b/Engine/Plugins/MovieScene/MovieRenderPipeline/Source/MovieRenderPipelineRenderPasses/Public/Graph/Nodes/MovieGraphImagePassBaseNode.h
@@ -131,6 +131,7 @@ protected:
 	UE_API virtual void TeardownImpl() override;
 	UE_API virtual void RenderImpl(const FMovieGraphTraversalContext& InFrameTraversalContext, const FMovieGraphTimeStepData& InTimeData) override;
 	UE_API virtual void GatherOutputPassesImpl(UMovieGraphEvaluatedConfig* InConfig, TArray<FMovieGraphRenderDataIdentifier>& OutExpectedPasses) const override;
+	UE_API virtual void GatherRenderCameraSourcesImpl(TArray<TSharedRef<FMovieGraphRenderCameraSource>>& OutSources) const override;
 	virtual int32 GetNumSceneViewsRenderedImpl() const { return CurrentInstances.Num(); }
     // ~UMovieGraphRenderPassNode Interface
 
diff --git a/Engine/Plugins/MovieScene/MovieRenderPipeline/Source/MovieRenderPipelineRenderPasses/Public/Graph/Renderers/MovieGraphDeferredPanoramicPass.h b/Engine/Plugins/MovieScene/MovieRenderPipeline/Source/MovieRenderPipelineRenderPasses/Public/Graph/Renderers/MovieGraphDeferredPanoramicPass.h
--- a/Engine/Plugins/MovieScene/MovieRenderPipeline/Source/MovieRenderPipelineRenderPasses/Public/Graph/Renderers/MovieGraphDeferredPanoramicPass.h
+++ b/Engine/Plugins/MovieScene/MovieRenderPipeline/Source/MovieRenderPipelineRenderPasses/Public/Graph/Renderers/MovieGraphDeferredPanoramicPass.h
@@ -24,6 +24,7 @@ namespace UE::MovieGraph::Rendering
 		UE_API virtual FAccumulatorSampleFunc GetAccumulateSampleFunction() const override;
 		UE_API virtual void GatherOutputPasses(UMovieGraphEvaluatedConfig* InConfig, TArray<FMovieGraphRenderDataIdentifier>& OutExpectedPasses) const override;
 		UE_API virtual FName GetBranchName() const;
+		UE_API virtual int32 GetRenderCameraIndex() const override;
 		// ~FMovieGraphImagePassBase Interface
 
 		// FMovieGraphDeferredPass Interface
diff --git a/Engine/Plugins/MovieScene/MovieRenderPipeline/Source/MovieRenderPipelineRenderPasses/Public/Graph/Renderers/MovieGraphDeferredPass.h b/Engine/Plugins/MovieScene/MovieRenderPipeline/Source/MovieRenderPipelineRenderPasses/Public/Graph/Renderers/MovieGraphDeferredPass.h
--- a/Engine/Plugins/MovieScene/MovieRenderPipeline/Source/MovieRenderPipelineRenderPasses/Public/Graph/Renderers/MovieGraphDeferredPass.h
+++ b/Engine/Plugins/MovieScene/MovieRenderPipeline/Source/MovieRenderPipelineRenderPasses/Public/Graph/Renderers/MovieGraphDeferredPass.h
@@ -27,6 +27,7 @@ namespace UE::MovieGraph::Rendering
 		UE_API virtual FName GetBranchName() const override;
 		UE_API virtual UMovieGraphImagePassBaseNode* GetParentNode(UMovieGraphEvaluatedConfig* InConfig) const override;
 		UE_API virtual bool ShouldDiscardOutput(const TSharedRef<FSceneViewFamilyContext>& InFamily, const UE::MovieGraph::DefaultRenderer::FCameraInfo& InCameraInfo) const override;
+		UE_API virtual int32 GetRenderCameraIndex() const override;
 		// End FMovieGraphImagePassBase
 
 	protected:
@@ -61,6 +62,9 @@ namespace UE::MovieGraph::Rendering
 		/** Provides the resolution and camera information that should be used for the next render. */
 		UE_API FResolutionAndCameraInfo GetResolutionAndCameraInfo(UMovieGraphEvaluatedConfig* InEvaluatedConfig) const;
 
+		/** Called after per-sample metadata is populated. Override to inject pass-specific metadata into the output. */
+		UE_API virtual void GetMetadataFromRenderPass(TMap<FString, FString>& InOutMetadataMap) const { };
+
 		/**
 		 * Encapsulates the shared view creation pipeline: GetSceneViewState -> CreateSceneViewFamily ->
 		 * CreateViewInitOptions -> CalculateProjectionMatrix -> ModifyProjectionMatrixForTiling ->
diff --git a/Engine/Plugins/MovieScene/MovieRenderPipeline/Source/MovieRenderPipelineRenderPasses/Public/Graph/Renderers/MovieGraphImagePassBase.h b/Engine/Plugins/MovieScene/MovieRenderPipeline/Source/MovieRenderPipelineRenderPasses/Public/Graph/Renderers/MovieGraphImagePassBase.h
--- a/Engine/Plugins/MovieScene/MovieRenderPipeline/Source/MovieRenderPipelineRenderPasses/Public/Graph/Renderers/MovieGraphImagePassBase.h
+++ b/Engine/Plugins/MovieScene/MovieRenderPipeline/Source/MovieRenderPipelineRenderPasses/Public/Graph/Renderers/MovieGraphImagePassBase.h
@@ -3,6 +3,7 @@
 #pragma once
 #include "CoreTypes.h"
 #include "Graph/MovieGraphDefaultRenderer.h"
+#include "Graph/MovieGraphRenderCameraSource.h"
 #include "SceneView.h"
 
 #define UE_API MOVIERENDERPIPELINERENDERPASSES_API
@@ -72,16 +73,32 @@ namespace UE::MovieGraph::Rendering
 		virtual ~FMovieGraphImagePassBase() = default;
 
 		UE_API virtual void Setup(TWeakObjectPtr<UMovieGraphDefaultRenderer> InRenderer, TWeakObjectPtr<UMovieGraphImagePassBaseNode> InRenderPassNode, const FMovieGraphRenderPassLayerData& InLayer);
-		virtual void Teardown() {}
+		UE_API virtual void Teardown();
+
 		virtual void Render(const FMovieGraphTraversalContext& InFrameTraversalContext, const FMovieGraphTimeStepData& InTimeData) {}
 		virtual void GatherOutputPasses(UMovieGraphEvaluatedConfig* InConfig, TArray<FMovieGraphRenderDataIdentifier>& OutExpectedPasses) const {}
-		virtual void AddReferencedObjects(FReferenceCollector& Collector) {}
+		UE_API virtual void AddReferencedObjects(FReferenceCollector& Collector);
 		virtual FName GetBranchName() const { return NAME_None; }
 		virtual TWeakObjectPtr<UMovieGraphDefaultRenderer> GetRenderer() const { return WeakGraphRenderer; }
 
+		/** Returns the shared camera source for this pass. Use when a shared reference is needed to extend lifetime. */
+		virtual TSharedPtr<const FMovieGraphRenderCameraSource> GetRenderCameraSource() const { return RenderCameraSource; }
+		virtual TSharedPtr<FMovieGraphRenderCameraSource> GetRenderCameraSource() { return RenderCameraSource; }
+
 	protected:
 		typedef TFunction<void(TUniquePtr<FImagePixelData>&&, const TSharedRef<FMovieGraphSampleState>, const TSharedRef<::MoviePipeline::IMoviePipelineAccumulationArgs>)> FAccumulatorSampleFunc;
 		
+		/** Returns the camera index used for this pass.
+		 * Returns [0..N) when using multi-camera rendering; otherwise INDEX_NONE.
+		 */
+		virtual int32 GetRenderCameraIndex() const { return INDEX_NONE; }
+
+		/** Returns the camera info for the given camera index, forwarding this pass's camera source to the renderer. */
+		UE_API virtual UE::MovieGraph::DefaultRenderer::FCameraInfo GetRenderCameraInfo() const;
+
+		/** Returns the overscan fraction for the given camera index, forwarding this pass's camera source to the renderer. */
+		UE_API virtual float GetRenderCameraOverscan() const;
+
 		/** Utility function for calculating a Projection Matrix (Orthographic or Perspective), and modify it based on the overscan percentage, aspect ratio, etc. */
 		UE_API virtual void CalculateProjectionMatrix(UE::MovieGraph::DefaultRenderer::FCameraInfo& InOutCameraInfo, FSceneViewProjectionData& InOutProjectionData, const FIntPoint InBackbufferResolution, const FIntPoint InAccumulatorResolution) const;
 		
@@ -122,6 +139,18 @@ namespace UE::MovieGraph::Rendering
 
 	protected:
 		TWeakObjectPtr<UMovieGraphDefaultRenderer> WeakGraphRenderer;
+
+		/**
+		 * Optional camera source shared among all FMovieGraphImagePassBase instances spawned by the same
+		 * UMovieGraphRenderPassNode. The node creates at most one FMovieGraphRenderCameraSource and hands
+		 * the same TSharedPtr to every pass instance it owns. Null when using default sequence cameras.
+		 *
+		 * FMovieGraphRenderCameraSource is a non-UObject, so the GC cannot trace into it. Each pass
+		 * instance forwards AddReferencedObjects() to the shared source, so the source's UObject
+		 * references are reported once per pass instance. Duplicate reports are safe; Unreal's GC
+		 * handles them correctly.
+		 */
+		TSharedPtr<FMovieGraphRenderCameraSource> RenderCameraSource;
 	};
 }
 
diff --git a/Engine/Plugins/VirtualProduction/AccumulationDOF/Source/AccumulationDOF/Private/AccumulationDOF.cpp b/Engine/Plugins/VirtualProduction/AccumulationDOF/Source/AccumulationDOF/Private/AccumulationDOF.cpp
--- a/Engine/Plugins/VirtualProduction/AccumulationDOF/Source/AccumulationDOF/Private/AccumulationDOF.cpp
+++ b/Engine/Plugins/VirtualProduction/AccumulationDOF/Source/AccumulationDOF/Private/AccumulationDOF.cpp
@@ -53,7 +53,11 @@ void FAccumulationDOFModule::StartupModule()
 			return nullptr;
 		}
 
-		const UE::MovieGraph::DefaultRenderer::FCameraInfo CameraInfo = InContext.Renderer->GetCameraInfo(InContext.LayerData.CameraIndex);
+		const TSharedPtr<const FMovieGraphRenderCameraSource> CameraSource = InContext.LayerData.RenderPassNode.IsValid()
+			? InContext.LayerData.RenderPassNode->GetOrCreateRenderCameraSource(InContext.Renderer->GetOwningGraph())
+			: nullptr;
+
+		const UE::MovieGraph::DefaultRenderer::FCameraInfo CameraInfo = InContext.Renderer->GetCameraInfo(InContext.LayerData.CameraIndex, CameraSource.Get());
 		const ACineCameraActor* CineCamera = Cast<ACineCameraActor>(CameraInfo.ViewActor);
 		if (!CineCamera)
 		{
diff --git a/Engine/Plugins/VirtualProduction/AccumulationDOF/Source/AccumulationDOF/Private/MovieGraphAccumulationDOFPass.cpp b/Engine/Plugins/VirtualProduction/AccumulationDOF/Source/AccumulationDOF/Private/MovieGraphAccumulationDOFPass.cpp
--- a/Engine/Plugins/VirtualProduction/AccumulationDOF/Source/AccumulationDOF/Private/MovieGraphAccumulationDOFPass.cpp
+++ b/Engine/Plugins/VirtualProduction/AccumulationDOF/Source/AccumulationDOF/Private/MovieGraphAccumulationDOFPass.cpp
@@ -1,4 +1,4 @@
-﻿// Copyright Epic Games, Inc. All Rights Reserved.
+// Copyright Epic Games, Inc. All Rights Reserved.
 
 #include "MovieGraphAccumulationDOFPass.h"
 
@@ -163,7 +163,7 @@ void FMovieGraphAccumulationDOFPass::Render(const FMovieGraphTraversalContext& I
 		FRenderTarget* RenderTargetResource = RenderTarget->GameThread_GetRenderTargetResource();
 		check(RenderTargetResource);
 
-		UE::MovieGraph::DefaultRenderer::FCameraInfo CameraInfo = GraphRenderer->GetCameraInfo(LayerData.CameraIndex);
+		UE::MovieGraph::DefaultRenderer::FCameraInfo CameraInfo = GetRenderCameraInfo();
 		CameraInfo.bAllowCameraAspectRatio = true;
 		CameraInfo.TilingParams.TileCount = FIntPoint(1, 1);
 		CameraInfo.TilingParams.TileIndexes = FIntPoint(0, 0);
@@ -504,7 +504,7 @@ void FMovieGraphAccumulationDOFPass::RenderSingleApertureSample(
 	check(RenderTargetResource);
 
 	// Build camera info
-	UE::MovieGraph::DefaultRenderer::FCameraInfo CameraInfo = GraphRenderer->GetCameraInfo(LayerData.CameraIndex);
+	UE::MovieGraph::DefaultRenderer::FCameraInfo CameraInfo = GetRenderCameraInfo();
 	CameraInfo.bAllowCameraAspectRatio = true;
 	CameraInfo.TilingParams.TileSize = InResolutionInfo.TileSize;
 	CameraInfo.TilingParams.TileCount = FIntPoint(1, 1);
@@ -665,7 +665,7 @@ void FMovieGraphAccumulationDOFPass::RenderInjectionAndSubmit(
 	check(RenderTargetResource);
 
 	// Build camera info for injection
-	UE::MovieGraph::DefaultRenderer::FCameraInfo CameraInfo = GraphRenderer->GetCameraInfo(LayerData.CameraIndex);
+	UE::MovieGraph::DefaultRenderer::FCameraInfo CameraInfo = GetRenderCameraInfo();
 	CameraInfo.bAllowCameraAspectRatio = true;
 	CameraInfo.TilingParams.TileSize = InResolutionInfo.TileSize;
 	CameraInfo.TilingParams.TileCount = FIntPoint(1, 1);
@@ -851,7 +851,7 @@ void FMovieGraphAccumulationDOFPass::InitializeSampler(UMovieGraphEvaluatedConfi
 
 UCineCameraComponent* FMovieGraphAccumulationDOFPass::FindCineCameraComponent() const
 {
-	const UE::MovieGraph::DefaultRenderer::FCameraInfo CameraInfo = GetRenderer()->GetCameraInfo(LayerData.CameraIndex);
+	const UE::MovieGraph::DefaultRenderer::FCameraInfo CameraInfo = GetRenderCameraInfo();
 
 	const ACineCameraActor* CineCameraActor = Cast<ACineCameraActor>(CameraInfo.ViewActor);
 	if (!CineCameraActor)
@@ -869,7 +869,7 @@ UAccumulationDOFComponent* FMovieGraphAccumulationDOFPass::FindAccumulationDOFCo
 		return CachedDOFComponent.Get();
 	}
 
-	const UE::MovieGraph::DefaultRenderer::FCameraInfo CameraInfo = GetRenderer()->GetCameraInfo(LayerData.CameraIndex);
+	const UE::MovieGraph::DefaultRenderer::FCameraInfo CameraInfo = GetRenderCameraInfo();
 
 	const ACineCameraActor* CineCameraActor = Cast<ACineCameraActor>(CameraInfo.ViewActor);
 	if (!CineCameraActor)
