MoviePipeline: Added the UI for the new "Basic" configuration type in the queue. This also includes a visual refresh of the queue.

Note that this is currently locked behind the following cvar temporarily (which must be set via the cmdline for all updates to be applied) - MovieRenderPipeline.ShowNewQueuePanel

#Synced-CL 52423999

diff --git a/Engine/Plugins/MovieScene/MovieRenderPipeline/Source/MovieRenderPipelineCore/Private/Graph/MovieRenderGraphEditorSettings.cpp b/Engine/Plugins/MovieScene/MovieRenderPipeline/Source/MovieRenderPipelineCore/Private/Graph/MovieRenderGraphEditorSettings.cpp
--- a/Engine/Plugins/MovieScene/MovieRenderPipeline/Source/MovieRenderPipelineCore/Private/Graph/MovieRenderGraphEditorSettings.cpp
+++ b/Engine/Plugins/MovieScene/MovieRenderPipeline/Source/MovieRenderPipelineCore/Private/Graph/MovieRenderGraphEditorSettings.cpp
@@ -7,7 +7,13 @@
 UMovieRenderGraphEditorSettings::UMovieRenderGraphEditorSettings()
 	: PostRenderSettings(FMovieGraphPostRenderSettings())
 {
-	
+	// Seed default visible output types for Basic Config mode.
+	// These defaults are used on first launch before the user modifies via the gear menu.
+	VisibleBasicConfigOutputTypes.Add(FSoftClassPath(TEXT("/Script/MovieRenderPipelineRenderPasses.MovieGraphImageSequenceOutputNode_MultiLayerEXR")));
+	VisibleBasicConfigOutputTypes.Add(FSoftClassPath(TEXT("/Script/MovieRenderPipelineRenderPasses.MovieGraphImageSequenceOutputNode_JPG")));
+	VisibleBasicConfigOutputTypes.Add(FSoftClassPath(TEXT("/Script/MovieRenderPipelineMP4Encoder.MovieGraphMP4EncoderNode")));
+	VisibleBasicConfigOutputTypes.Add(FSoftClassPath(TEXT("/Script/MovieRenderPipelineRenderPasses.MovieGraphImageSequenceOutputNode_PNG")));
+	VisibleBasicConfigOutputTypes.Add(FSoftClassPath(TEXT("/Script/AppleProResMedia.MovieGraphAppleProResNode")));
 }
 
 FName UMovieRenderGraphEditorSettings::GetCategoryName() const
diff --git a/Engine/Plugins/MovieScene/MovieRenderPipeline/Source/MovieRenderPipelineCore/Private/MoviePipelineQueue.cpp b/Engine/Plugins/MovieScene/MovieRenderPipeline/Source/MovieRenderPipelineCore/Private/MoviePipelineQueue.cpp
--- a/Engine/Plugins/MovieScene/MovieRenderPipeline/Source/MovieRenderPipelineCore/Private/MoviePipelineQueue.cpp
+++ b/Engine/Plugins/MovieScene/MovieRenderPipeline/Source/MovieRenderPipelineCore/Private/MoviePipelineQueue.cpp
@@ -295,6 +295,8 @@ void UMoviePipelineExecutorJob::SetGraphPreset(const UMovieGraphConfig* InGraphP
 
 	GraphPreset = const_cast<UMovieGraphConfig*>(InGraphPreset);
 
+	const EMoviePipelineConfigMode PreviousMode = ConfigMode;
+
 	if (InGraphPreset != nullptr)
 	{
 		ConfigMode = EMoviePipelineConfigMode::Graph;
@@ -317,6 +319,11 @@ void UMoviePipelineExecutorJob::SetGraphPreset(const UMovieGraphConfig* InGraphP
 	}
 
 	OnJobGraphPresetChanged.Broadcast(this, GraphPreset.Get());
+
+	if (PreviousMode != ConfigMode)
+	{
+		OnJobConfigModeChanged.Broadcast(this, ConfigMode);
+	}
 }
 
 void UMoviePipelineExecutorJob::SetSequence(FSoftObjectPath InSequence)
@@ -366,7 +373,14 @@ void UMoviePipelineExecutorJob::SetConfiguration(UMoviePipelinePrimaryConfig* In
 	{
 		Configuration->CopyFrom(InPreset);
 		PresetOrigin = nullptr;
+
+		const EMoviePipelineConfigMode PreviousMode = ConfigMode;
 		ConfigMode = EMoviePipelineConfigMode::Preset;
+
+		if (PreviousMode != ConfigMode)
+		{
+			OnJobConfigModeChanged.Broadcast(this, ConfigMode);
+		}
 	}
 }
 
@@ -376,7 +390,14 @@ void UMoviePipelineExecutorJob::SetPresetOrigin(UMoviePipelinePrimaryConfig* InP
 	{
 		Configuration->CopyFrom(InPreset);
 		PresetOrigin = TSoftObjectPtr<UMoviePipelinePrimaryConfig>(InPreset);
+		
+		const EMoviePipelineConfigMode PreviousMode = ConfigMode;
 		ConfigMode = EMoviePipelineConfigMode::Preset;
+		
+		if (PreviousMode != ConfigMode)
+		{
+			OnJobConfigModeChanged.Broadcast(this, ConfigMode);
+		}
 	}
 }
 
@@ -418,6 +439,7 @@ UMoviePipelineBasicConfig* UMoviePipelineExecutorJob::SetupBasicConfiguration()
 		}
 	}
 
+	const EMoviePipelineConfigMode PreviousMode = ConfigMode;
 	ConfigMode = EMoviePipelineConfigMode::Basic;
 
 	// Ensure all child shots have a BasicShotConfig so they can display override controls.
@@ -430,9 +452,51 @@ UMoviePipelineBasicConfig* UMoviePipelineExecutorJob::SetupBasicConfiguration()
 		}
 	}
 
+	if (PreviousMode != ConfigMode)
+	{
+		OnJobConfigModeChanged.Broadcast(this, ConfigMode);
+	}
+
 	return BasicConfig;
 }
 
+bool UMoviePipelineExecutorJob::SetConfigMode(EMoviePipelineConfigMode InConfigMode)
+{
+	if (InConfigMode == ConfigMode)
+	{
+		return true;
+	}
+
+	// Switching to Graph mode requires a graph to already be assigned. Callers are responsible for assigning a graph
+	// first (e.g. via SetGraphPreset) if one doesn't exist, since loading the default graph requires editor-only
+	// project settings.
+	if (InConfigMode == EMoviePipelineConfigMode::Graph && GraphPreset.IsNull())
+	{
+		return false;
+	}
+
+#if WITH_EDITOR
+	Modify();
+#endif
+
+	const EMoviePipelineConfigMode PreviousMode = ConfigMode;
+	ConfigMode = InConfigMode;
+
+	// Lazily initialize Basic config and shot overrides when entering Basic mode for the first time.
+	// SetupBasicConfiguration handles creating the config from saved defaults and setting up shot overrides.
+	if (ConfigMode == EMoviePipelineConfigMode::Basic)
+	{
+		SetupBasicConfiguration();
+	}
+
+	if (PreviousMode != ConfigMode)
+	{
+		OnJobConfigModeChanged.Broadcast(this, ConfigMode);
+	}
+
+	return true;
+}
+
 bool UMoviePipelineExecutorJob::CanUseConsoleVariableOverrides() const
 {
 	return IsUsingGraphConfiguration() || IsUsingBasicConfiguration();
diff --git a/Engine/Plugins/MovieScene/MovieRenderPipeline/Source/MovieRenderPipelineCore/Public/Graph/MovieRenderGraphEditorSettings.h b/Engine/Plugins/MovieScene/MovieRenderPipeline/Source/MovieRenderPipelineCore/Public/Graph/MovieRenderGraphEditorSettings.h
--- a/Engine/Plugins/MovieScene/MovieRenderPipeline/Source/MovieRenderPipelineCore/Public/Graph/MovieRenderGraphEditorSettings.h
+++ b/Engine/Plugins/MovieScene/MovieRenderPipeline/Source/MovieRenderPipelineCore/Public/Graph/MovieRenderGraphEditorSettings.h
@@ -3,6 +3,7 @@
 #pragma once
 
 #include "Engine/DeveloperSettings.h"
+#include "UObject/SoftObjectPath.h"
 
 #include "MovieGraphDataTypes.h"
 #include "MovieGraphQuickRenderSettings.h"
@@ -12,6 +13,15 @@
 
 #define UE_API MOVIERENDERPIPELINECORE_API
 
+/** The default configuration type assigned to newly created jobs in the queue. */
+UENUM()
+enum class EMoviePipelineDefaultConfigType : uint8
+{
+	Basic,
+	MovieRenderGraph,
+	LegacyPreset
+};
+
 UCLASS(MinimalAPI, BlueprintType, Config=EditorPerProjectUserSettings, meta=(DisplayName="Movie Render Graph"))
 class UMovieRenderGraphEditorSettings : public UDeveloperSettings
 {
@@ -29,6 +39,21 @@ public:
 	/** If playing back media post-render, these settings dictate how that will be done. */
 	UPROPERTY(Config, EditAnywhere, Category="Play Render Output", meta=(ShowOnlyInnerProperties))
 	FMovieGraphPostRenderSettings PostRenderSettings;
+
+	/** The configuration type assigned to newly created jobs within the editor. */
+	UPROPERTY(Config, EditAnywhere, Category="Queue Defaults")
+	EMoviePipelineDefaultConfigType DefaultConfigType = EMoviePipelineDefaultConfigType::Basic;
+
+	/**
+	 * Output types visible as toggle buttons in Basic Config mode.
+	 * Each entry is the soft class path of a UMovieGraphFileOutputNode subclass.
+	 */
+	UPROPERTY(Config)
+	TArray<FSoftClassPath> VisibleBasicConfigOutputTypes;
+
+	/** Whether the Burn In Class property is shown in Basic Config mode. */
+	UPROPERTY(Config)
+	bool bShowBurnInForBasicConfig = false;
 };
 
 #undef UE_API
diff --git a/Engine/Plugins/MovieScene/MovieRenderPipeline/Source/MovieRenderPipelineCore/Public/MoviePipelineQueue.h b/Engine/Plugins/MovieScene/MovieRenderPipeline/Source/MovieRenderPipelineCore/Public/MoviePipelineQueue.h
--- a/Engine/Plugins/MovieScene/MovieRenderPipeline/Source/MovieRenderPipelineCore/Public/MoviePipelineQueue.h
+++ b/Engine/Plugins/MovieScene/MovieRenderPipeline/Source/MovieRenderPipelineCore/Public/MoviePipelineQueue.h
@@ -354,6 +354,7 @@ private:
 };
 
 DECLARE_MULTICAST_DELEGATE_TwoParams(FOnMoviePipelineJobGraphPresetChanged, UMoviePipelineExecutorJob*, UMovieGraphConfig*);
+DECLARE_MULTICAST_DELEGATE_TwoParams(FOnMoviePipelineJobConfigModeChanged, UMoviePipelineExecutorJob*, EMoviePipelineConfigMode);
 
 /**
 * A particular job within the Queue
@@ -537,6 +538,15 @@ public:
 	UFUNCTION(BlueprintPure, Category = "Movie Render Pipeline")
 	EMoviePipelineConfigMode GetConfigMode() const { return ConfigMode; }
 
+	/**
+	 * Switches the active configuration mode on the job.
+	 * For Basic mode, lazily initializes the BasicConfig if it doesn't exist yet.
+	 * For Graph mode, requires a graph preset to already be assigned (returns false otherwise). See SetGraphPreset().
+	 * @return true if the mode was switched (or was already the requested mode), false if the switch was not possible.
+	 */
+	UFUNCTION(BlueprintCallable, Category = "Movie Render Pipeline")
+	UE_API bool SetConfigMode(EMoviePipelineConfigMode InConfigMode);
+
 	/**
 	 * Returns true if this job is using graph-style configuration, else false.
 	 */
@@ -692,6 +702,9 @@ public:
 	/** Called when the graph preset assigned to the job changes. */
 	FOnMoviePipelineJobGraphPresetChanged OnJobGraphPresetChanged;
 
+	/** Called when the job's configuration mode changes (Basic, Graph, Preset). */
+	FOnMoviePipelineJobConfigModeChanged OnJobConfigModeChanged;
+
 	/**
 	 * (Optional) Console variable overrides which are applied after cvars set via nodes. Applies to graph-based and Basic configurations.
 	 */
diff --git a/Engine/Plugins/MovieScene/MovieRenderPipeline/Source/MovieRenderPipelineEditor/Private/Customizations/JobCustomization_NewQueue.cpp b/Engine/Plugins/MovieScene/MovieRenderPipeline/Source/MovieRenderPipelineEditor/Private/Customizations/JobCustomization_NewQueue.cpp
new file mode 100644
--- /dev/null
+++ b/Engine/Plugins/MovieScene/MovieRenderPipeline/Source/MovieRenderPipelineEditor/Private/Customizations/JobCustomization_NewQueue.cpp
@@ -0,0 +1,975 @@
+// Copyright Epic Games, Inc. All Rights Reserved.
+
+#include "Customizations/JobCustomization_NewQueue.h"
+
+// MovieRenderPipeline
+#include "Graph/MovieGraphConfig.h"
+#include "Graph/MovieGraphCustomizationUtils.h"
+#include "Graph/MovieRenderGraphEditorSettings.h"
+#include "Graph/Nodes/MovieGraphPathTracerPassNode.h"
+#include "MoviePipelineBasicConfig.h"
+#include "MoviePipelineEditorBlueprintLibrary.h"
+#include "MoviePipelinePrimaryConfig.h"
+#include "MoviePipelineQueue.h"
+#include "Widgets/SMoviePipelineQueueEditor.h"
+
+// Engine / PropertyEditor
+#include "DetailCategoryBuilder.h"
+#include "DetailLayoutBuilder.h"
+#include "DetailWidgetRow.h"
+#include "IDetailGroup.h"
+#include "PropertyCustomizationHelpers.h"
+#include "PropertyHandle.h"
+
+// Framework
+#include "Framework/MultiBox/MultiBoxBuilder.h"
+#include "ScopedTransaction.h"
+#include "Styling/AppStyle.h"
+#include "UObject/Package.h"
+#include "UObject/UObjectIterator.h"
+
+// Slate / Widgets
+#include "Widgets/Images/SImage.h"
+#include "Widgets/Input/SCheckBox.h"
+#include "Widgets/Input/SComboButton.h"
+#include "Widgets/Input/SSegmentedControl.h"
+#include "Widgets/Layout/SWrapBox.h"
+#include "Widgets/Text/STextBlock.h"
+
+// Note: this static_assert needs to be done in the customization rather than where EMoviePipelineBasicDenoiserType is defined because it's
+// defined in a module that cannot access EMovieGraphPathTracerDenoiserType.
+static_assert(
+	static_cast<uint8>(EMoviePipelineBasicDenoiserType::Spatial) == static_cast<uint8>(EMovieGraphPathTracerDenoiserType::Spatial)
+	&& static_cast<uint8>(EMoviePipelineBasicDenoiserType::Temporal) == static_cast<uint8>(EMovieGraphPathTracerDenoiserType::Temporal),
+	"EMoviePipelineBasicDenoiserType must match EMovieGraphPathTracerDenoiserType");
+
+#define LOCTEXT_NAMESPACE "MoviePipelineEditor"
+
+struct FToggleButtonDef
+{
+	FText Label;
+	bool* Value;
+	// Optional: called when the user tries to uncheck this button. Return false to prevent unchecking.
+	// Use this for cross-button constraints (e.g. "at least one renderer must be active").
+	TFunction<bool()> CanUncheck;
+};
+
+/**
+ * Adds a row of toggle-style checkboxes to a detail category.
+ * When InGetRowOverride is set (shot mode), a single override checkbox appears in the name column
+ * controlling the whole row — all buttons and the row label are disabled when the override is off.
+ */
+static void AddToggleButtonRow(
+	IDetailCategoryBuilder& InCategory,
+	UMoviePipelineBasicConfig* InConfig,
+	const FText& InFilterText,
+	const FText& InRowLabel,
+	TArrayView<const FToggleButtonDef> InButtons,
+	TFunction<bool()> InGetRowOverride = {},
+	TFunction<void(bool)> InSetRowOverride = {})
+{
+	const bool bHasRowOverride = !!InGetRowOverride;
+
+	// When a row override is present, all buttons and the label are enabled only when the override is active.
+	TAttribute<bool> RowEnabled = bHasRowOverride
+		? TAttribute<bool>::CreateLambda([InGetRowOverride]() { return InGetRowOverride(); })
+		: TAttribute<bool>(true);
+
+	TSharedRef<SHorizontalBox> ButtonBox = SNew(SHorizontalBox);
+	for (const FToggleButtonDef& Button : InButtons)
+	{
+		bool* BoolPtr = Button.Value;
+		TFunction<bool()> CanUncheck = Button.CanUncheck;
+		ButtonBox->AddSlot()
+			.AutoWidth()
+			.Padding(0, 0, 5, 0)
+			[
+				SNew(SCheckBox)
+				.Style(FAppStyle::Get(), "DetailsView.SectionButton")
+				.IsEnabled(RowEnabled)
+				.IsChecked_Lambda([BoolPtr]()
+				{
+					return *BoolPtr ? ECheckBoxState::Checked : ECheckBoxState::Unchecked;
+				})
+				.OnCheckStateChanged_Lambda([InConfig, BoolPtr, CanUncheck](ECheckBoxState NewState)
+				{
+					if (NewState == ECheckBoxState::Unchecked && CanUncheck && !CanUncheck())
+					{
+						return;
+					}
+					InConfig->Modify();
+					*BoolPtr = (NewState == ECheckBoxState::Checked);
+				})
+				[
+					SNew(STextBlock)
+					.TextStyle(FAppStyle::Get(), "SmallText")
+					.Text(Button.Label)
+				]
+			];
+	}
+
+	InCategory.AddCustomRow(InFilterText)
+		.NameContent()
+		[
+			SNew(SHorizontalBox)
+			+ SHorizontalBox::Slot()
+			.AutoWidth()
+			.VAlign(VAlign_Center)
+			.Padding(-2, 0, 2, 0)
+			[
+				// Single override toggle for the whole row — visible in shot mode only.
+				SNew(SCheckBox)
+				.Visibility(bHasRowOverride ? EVisibility::Visible : EVisibility::Collapsed)
+				.IsChecked_Lambda([bHasRowOverride, InGetRowOverride]()
+				{
+					return (bHasRowOverride && InGetRowOverride())
+						? ECheckBoxState::Checked : ECheckBoxState::Unchecked;
+				})
+				.OnCheckStateChanged_Lambda([InConfig, bHasRowOverride, InSetRowOverride](ECheckBoxState NewState)
+				{
+					if (bHasRowOverride)
+					{
+						InConfig->Modify();
+						InSetRowOverride(NewState == ECheckBoxState::Checked);
+					}
+				})
+			]
+			+ SHorizontalBox::Slot()
+			.AutoWidth()
+			.VAlign(VAlign_Center)
+			[
+				SNew(STextBlock)
+				.Text(InRowLabel)
+				.Font(IDetailLayoutBuilder::GetDetailFont())
+				.IsEnabled(RowEnabled)
+			]
+		]
+		.ValueContent()
+		[ ButtonBox ];
+}
+
+/**
+ * Wires an inline override toggle to a property row for shots.
+ * For primary jobs (bIsShot == false) this is a no-op -- all properties are always editable.
+ */
+static void OptionallyAddShotInlineEditCondition(
+	IDetailPropertyRow* InRow,
+	UMoviePipelineBasicConfig* InConfig,
+	const bool bIsShot,
+	TFunction<bool()> InGetFlag,
+	TFunction<void(bool)> InSetFlag)
+{
+	if (!bIsShot || !InRow)
+	{
+		return;
+	}
+
+	InRow->EditCondition(
+		TAttribute<bool>::CreateLambda([InGetFlag]() -> bool
+		{
+			return InGetFlag();
+		}),
+		FOnBooleanValueChanged::CreateLambda([InConfig, InSetFlag](const bool bNew)
+		{
+			InConfig->Modify();
+			InSetFlag(bNew);
+		})
+	);
+}
+
+TSharedRef<IDetailCustomization> FJobDetailsCustomization_NewQueue::MakeInstance()
+{
+	return MakeShared<FJobDetailsCustomization_NewQueue>();
+}
+
+void FJobDetailsCustomization_NewQueue::PendingDelete()
+{
+	// Unregister delegates. It's important to do this in PendingDelete() vs the destructor because the destructor is not called before the next
+	// details panel is created (via ForceRefreshDetails()), leading to an exponential increase in the number of delegates registered.
+
+	UPackage::PackageSavedWithContextEvent.RemoveAll(this);
+
+	if (SelectedJob.IsValid())
+	{
+		SelectedJob->OnJobGraphPresetChanged.RemoveAll(this);
+		SelectedJob->OnJobConfigModeChanged.RemoveAll(this);
+	}
+
+	if (SelectedShot.IsValid())
+	{
+		SelectedShot->OnShotGraphPresetChanged.RemoveAll(this);
+	}
+}
+
+void FJobDetailsCustomization_NewQueue::CustomizeDetails(const TSharedPtr<IDetailLayoutBuilder>& InDetailBuilder)
+{
+	DetailBuilder = InDetailBuilder.Get();
+	CustomizeDetails(*InDetailBuilder);
+}
+
+void FJobDetailsCustomization_NewQueue::RefreshLayout(const FString&, UPackage*, FObjectPostSaveContext) const
+{
+	if (DetailBuilder) { DetailBuilder->ForceRefreshDetails(); }
+}
+
+void FJobDetailsCustomization_NewQueue::RefreshLayout(UMoviePipelineExecutorShot*, UMovieGraphConfig*) const
+{
+	if (DetailBuilder) { DetailBuilder->ForceRefreshDetails(); }
+}
+
+void FJobDetailsCustomization_NewQueue::RefreshLayout(UMoviePipelineExecutorJob*, UMovieGraphConfig*) const
+{
+	if (DetailBuilder) { DetailBuilder->ForceRefreshDetails(); }
+}
+
+void FJobDetailsCustomization_NewQueue::RefreshLayout(UMoviePipelineExecutorJob*, EMoviePipelineConfigMode) const
+{
+	if (DetailBuilder) { DetailBuilder->ForceRefreshDetails(); }
+}
+
+void FJobDetailsCustomization_NewQueue::CustomizeDetails(IDetailLayoutBuilder& InDetailBuilder)
+{
+	DetailBuilder = &InDetailBuilder;
+
+	// Refresh the customization every time a save happens. Use this opportunity to update the variables in the UI. We could update the UI before
+	// a save occurs, but this would be very difficult to get right when multiple subgraphs are involved.
+	UPackage::PackageSavedWithContextEvent.AddSP(this, &FJobDetailsCustomization_NewQueue::RefreshLayout);
+
+	TArray<TWeakObjectPtr<UObject>> ObjectsBeingCustomized;
+	InDetailBuilder.GetObjectsBeingCustomized(ObjectsBeingCustomized);
+
+	TArray<UMoviePipelineExecutorJob*> SelectedJobs;
+	TArray<UMoviePipelineExecutorShot*> SelectedShots;
+
+	for (const TWeakObjectPtr<UObject>& SelectedObject : ObjectsBeingCustomized)
+	{
+		if (UMoviePipelineExecutorJob* SelectedJobTemp = Cast<UMoviePipelineExecutorJob>(SelectedObject.Get()))
+		{
+			SelectedJobs.Add(SelectedJobTemp);
+		}
+		else if (UMoviePipelineExecutorShot* SelectedShotTemp = Cast<UMoviePipelineExecutorShot>(SelectedObject.Get()))
+		{
+			SelectedShots.Add(SelectedShotTemp);
+		}
+	}
+
+	// Hide the original assignments properties (since they present an asset picker by default) for both jobs and shots
+	const TSharedRef<IPropertyHandle> JobAssignmentsProperty =
+		InDetailBuilder.GetProperty(TEXT("GraphVariableAssignments"), UMoviePipelineExecutorJob::StaticClass());
+	const TSharedRef<IPropertyHandle> ShotAssignmentsProperty =
+		InDetailBuilder.GetProperty(TEXT("GraphVariableAssignments"), UMoviePipelineExecutorShot::StaticClass());
+	const TSharedRef<IPropertyHandle> ShotPrimaryGraphAssignmentsProperty =
+		InDetailBuilder.GetProperty(TEXT("PrimaryGraphVariableAssignments"), UMoviePipelineExecutorShot::StaticClass());
+	InDetailBuilder.HideProperty(JobAssignmentsProperty);
+	InDetailBuilder.HideProperty(ShotAssignmentsProperty);
+	InDetailBuilder.HideProperty(ShotPrimaryGraphAssignmentsProperty);
+	InDetailBuilder.HideProperty(
+		InDetailBuilder.GetProperty(GET_MEMBER_NAME_CHECKED(UMoviePipelineExecutorJob, JobName)));
+	InDetailBuilder.HideProperty(
+		InDetailBuilder.GetProperty(GET_MEMBER_NAME_CHECKED(UMoviePipelineExecutorJob, Sequence)));
+
+	// Only display the customized variables UI if there is one job or shot selected
+	const bool bIsPrimaryJob = (SelectedJobs.Num() == 1) && SelectedShots.IsEmpty();
+	const bool bIsShot = (SelectedShots.Num() == 1) && SelectedJobs.IsEmpty();
+	if (!bIsPrimaryJob && !bIsShot)
+	{
+		return;
+	}
+
+	// Refresh the UI if the graph preset changes (so the new variable assignments are displayed)
+	if (bIsShot)
+	{
+		SelectedShot = SelectedShots[0];
+		SelectedShot->OnShotGraphPresetChanged.AddSP(this, &FJobDetailsCustomization_NewQueue::RefreshLayout);
+
+		// Also listen for changes to the primary job. Changes to the primary job can trigger an update to shot variable assignments.
+		if (UMoviePipelineExecutorJob* PrimaryJob = SelectedShot->GetTypedOuter<UMoviePipelineExecutorJob>())
+		{
+			SelectedJob = PrimaryJob;
+			SelectedJob->OnJobGraphPresetChanged.AddSP(this, &FJobDetailsCustomization_NewQueue::RefreshLayout);
+			SelectedJob->OnJobConfigModeChanged.AddSP(this, &FJobDetailsCustomization_NewQueue::RefreshLayout);
+		}
+	}
+	else
+	{
+		SelectedJob = SelectedJobs[0];
+		SelectedJob->OnJobGraphPresetChanged.AddSP(this, &FJobDetailsCustomization_NewQueue::RefreshLayout);
+		SelectedJob->OnJobConfigModeChanged.AddSP(this, &FJobDetailsCustomization_NewQueue::RefreshLayout);
+	}
+
+	// Get the Basic config (if active) so we can surface its properties later.
+	UMoviePipelineBasicConfig* BasicConfig = nullptr;
+	if (SelectedJob.IsValid() && SelectedJob->IsUsingBasicConfiguration())
+	{
+		BasicConfig = SelectedJob->GetBasicConfig();
+	}
+
+	// Set up the categories for variable assignments.
+	IDetailCategoryBuilder& PrimaryGraphVariablesCategory = InDetailBuilder.EditCategory(
+		"PrimaryGraphVariables", LOCTEXT("PrimaryGraphVariablesCategory", "Primary Graph Variables"));
+	IDetailCategoryBuilder& PrimaryGraphVariablesShotOverridesCategory = InDetailBuilder.EditCategory(
+		"PrimaryGraphVariablesShotOverrides", LOCTEXT("PrimaryGraphVariablesShotOverridesCategory", "Primary Graph Variables (shot overrides)"));
+	IDetailCategoryBuilder& ShotGraphVariablesCategory = InDetailBuilder.EditCategory(
+		"ShotGraphVariables", LOCTEXT("ShotGraphVariablesCategory", "Shot Graph Variables"));
+
+	// Set all categories as hidden by default. Individual categories will be made visible if variables are added under them.
+	PrimaryGraphVariablesCategory.SetCategoryVisibility(false);
+	PrimaryGraphVariablesShotOverridesCategory.SetCategoryVisibility(false);
+	ShotGraphVariablesCategory.SetCategoryVisibility(false);
+
+	const bool bIsGraphMode = SelectedJob.IsValid() && SelectedJob->IsUsingGraphConfiguration();
+	if (bIsGraphMode)
+	{
+		if (bIsShot)
+		{
+			UE::MovieRenderPipelineEditor::Private::AddVariableAssignments(SelectedShot->GetGraphVariableAssignments(), ShotGraphVariablesCategory, DetailBuilder);
+			UE::MovieRenderPipelineEditor::Private::AddVariableAssignments(SelectedShot->GetPrimaryGraphVariableAssignments(), PrimaryGraphVariablesShotOverridesCategory, DetailBuilder);
+		}
+		else
+		{
+			UE::MovieRenderPipelineEditor::Private::AddVariableAssignments(SelectedJob->GetGraphVariableAssignments(), PrimaryGraphVariablesCategory, DetailBuilder);
+		}
+	}
+
+	IDetailCategoryBuilder& MovieRenderPipelineCategory = InDetailBuilder.EditCategory("Movie Render Pipeline");
+
+	// The Console Variable Overrides category should be collapsed by default. It's considered advanced, and most people won't need to use it.
+	IDetailCategoryBuilder& ConsoleVariableOverridesCategory = InDetailBuilder.EditCategory("Console Variable Overrides");
+	ConsoleVariableOverridesCategory.InitiallyCollapsed(true);
+
+	// Add Configuration mode row (only for primary jobs, not shots).
+	if (bIsPrimaryJob)
+	{
+		AddConfigurationModeRow(MovieRenderPipelineCategory);
+	}
+
+	// In Basic mode, surface Basic config properties organized into Output and Rendering categories.
+	// For primary jobs: display the job's BasicConfig (always editable).
+	// For shots: display the shot's BasicShotConfig with per-property override toggles.
+	// Guard on BasicConfig pointer — the enum may indicate Basic before the object is ready.
+	// ActiveConfig lifetime: held as a UPROPERTY on the job/shot, which the details panel references
+	// via TWeakObjectPtr. The panel is rebuilt (ForceRefreshDetails) on any mode/preset change, so
+	// these lambda captures are safe as long as the owning job/shot stays alive.
+	UMoviePipelineBasicConfig* ActiveConfig = nullptr;
+	if (BasicConfig && bIsPrimaryJob)
+	{
+		ActiveConfig = BasicConfig;
+	}
+	else if (BasicConfig && bIsShot && SelectedShot.IsValid())
+	{
+		ActiveConfig = SelectedShot->GetBasicShotConfig();
+	}
+
+	if (ActiveConfig)
+	{
+		AddBasicConfigOutputTypeProperties(InDetailBuilder, ActiveConfig, bIsShot);
+		AddBasicConfigRenderingProperties(InDetailBuilder, ActiveConfig, bIsShot);
+
+		// Re-acquire categories created by the helpers above
+		// (EditCategory returns the existing builder when called with the same ID).
+		IDetailCategoryBuilder& OutputCategory = InDetailBuilder.EditCategory("Output");
+		IDetailCategoryBuilder& RenderingCategory = InDetailBuilder.EditCategory("Rendering");
+
+		// Give the categories a specific ordering for Basic mode.
+		int32 SortOrder = 0;
+		MovieRenderPipelineCategory.SetSortOrder(SortOrder);
+		OutputCategory.SetSortOrder(++SortOrder);
+		RenderingCategory.SetSortOrder(++SortOrder);
+		ConsoleVariableOverridesCategory.SetSortOrder(++SortOrder);
+	}
+	else
+	{
+		// Give the categories a specific ordering for Graph/Legacy mode.
+		int32 SortOrder = 0;
+		MovieRenderPipelineCategory.SetSortOrder(SortOrder);
+		PrimaryGraphVariablesCategory.SetSortOrder(++SortOrder);
+		PrimaryGraphVariablesShotOverridesCategory.SetSortOrder(++SortOrder);
+		ShotGraphVariablesCategory.SetSortOrder(++SortOrder);
+		ConsoleVariableOverridesCategory.SetSortOrder(++SortOrder);
+	}
+}
+
+void FJobDetailsCustomization_NewQueue::AddConfigurationModeRow(IDetailCategoryBuilder& InCategory)
+{
+	IDetailLayoutBuilder* CachedBuilder = DetailBuilder;
+	TWeakObjectPtr<UMoviePipelineExecutorJob> WeakJob = SelectedJob;
+
+	// Lambda to switch config mode on the selected job.
+	// The OnJobConfigModeChanged delegate triggers a details panel refresh automatically.
+	auto SwitchConfigMode = [WeakJob](const EMoviePipelineConfigMode NewMode)
+	{
+		FScopedTransaction Transaction(LOCTEXT("SetConfigMode_Transaction", "Set Configuration Mode"));
+		if (UMoviePipelineExecutorJob* Job = WeakJob.Get())
+		{
+			if (!Job->SetConfigMode(NewMode))
+			{
+				// SetConfigMode fails for Graph when no graph is assigned; assign the default
+				// from project settings. AssignDefaultGraphPresetToJob calls SetGraphPreset,
+				// which internally switches ConfigMode to Graph.
+				SMoviePipelineQueueEditor::AssignDefaultGraphPresetToJob(Job);
+			}
+		}
+	};
+
+	// Configuration mode segmented control.
+	const TSharedRef<SWidget> ConfigButtonBox =
+		SNew(SSegmentedControl<EMoviePipelineConfigMode>)
+		.Value_Lambda([WeakJob]()
+		{
+			if (const UMoviePipelineExecutorJob* Job = WeakJob.Get())
+			{
+				return Job->GetConfigMode();
+			}
+			return EMoviePipelineConfigMode::Preset;
+		})
+		.OnValueChanged_Lambda([SwitchConfigMode](const EMoviePipelineConfigMode NewMode)
+		{
+			SwitchConfigMode(NewMode);
+		})
+		+ SSegmentedControl<EMoviePipelineConfigMode>::Slot(EMoviePipelineConfigMode::Basic)
+			.Text(LOCTEXT("ConfigMode_Basic", "Basic"))
+		+ SSegmentedControl<EMoviePipelineConfigMode>::Slot(EMoviePipelineConfigMode::Graph)
+			.Text(LOCTEXT("ConfigMode_Graph", "Graph"))
+		+ SSegmentedControl<EMoviePipelineConfigMode>::Slot(EMoviePipelineConfigMode::Preset)
+			.Text(LOCTEXT("ConfigMode_Preset", "Preset"));
+
+	// Build the mode-dependent sub-widget.
+	TSharedRef<SWidget> SubWidget = SNullWidget::NullWidget;
+
+	if (SelectedJob.IsValid() && SelectedJob->GetConfigMode() == EMoviePipelineConfigMode::Basic)
+	{
+		// Save... combo button with menu.
+		SubWidget = SNew(SComboButton)
+			.OnGetMenuContent_Lambda([WeakJob, CachedBuilder]() -> TSharedRef<SWidget>
+			{
+				FMenuBuilder MenuBuilder(true, nullptr);
+
+				MenuBuilder.BeginSection(NAME_None, LOCTEXT("SaveGraphSection", "Save Movie Render Graph"));
+				{
+					MenuBuilder.AddMenuEntry(
+						LOCTEXT("SaveAsGraphConfig", "Save As Movie Render Graph Config..."),
+						FText(),
+						FSlateIcon(),
+						FUIAction(FExecuteAction::CreateLambda([WeakJob, CachedBuilder]()
+						{
+							if (WeakJob.IsValid() && UMoviePipelineEditorBlueprintLibrary::SaveBasicConfigAsGraphConfig(WeakJob.Get()))
+							{
+								// SaveBasicConfigAsGraphConfig sets the graph preset on the job,
+								// which switches the config mode to Graph.
+								CachedBuilder->ForceRefreshDetails();
+							}
+						}))
+					);
+				}
+				MenuBuilder.EndSection();
+
+				MenuBuilder.BeginSection(NAME_None, LOCTEXT("SaveBasicSection", "Save Basic Config"));
+				{
+					MenuBuilder.AddMenuEntry(
+						LOCTEXT("SaveBasicAsDefault", "Save Basic Config As Default"),
+						LOCTEXT("SaveBasicAsDefaultTooltip", "Save the current Basic configuration as the default for new jobs."),
+						FSlateIcon(),
+						FUIAction(FExecuteAction::CreateLambda([WeakJob]()
+						{
+							if (const UMoviePipelineExecutorJob* Job = WeakJob.Get())
+							{
+								if (const UMoviePipelineBasicConfig* Config = Job->GetBasicConfig())
+								{
+									UMoviePipelineBasicConfig::SaveAsDefault(Config);
+								}
+							}
+						}))
+					);
+				}
+				MenuBuilder.EndSection();
+
+				return MenuBuilder.MakeWidget();
+			})
+			.ButtonContent()
+			[
+				SNew(SHorizontalBox)
+				+ SHorizontalBox::Slot()
+				.AutoWidth()
+				.VAlign(VAlign_Center)
+				.Padding(0, 0, 4, 0)
+				[
+					SNew(SImage)
+					.Image(FAppStyle::Get().GetBrush("AssetEditor.SaveAsset"))
+					.DesiredSizeOverride(FVector2D(16.f, 16.f))
+					.ColorAndOpacity(FSlateColor::UseForeground())
+				]
+				+ SHorizontalBox::Slot()
+				.AutoWidth()
+				.VAlign(VAlign_Center)
+				[
+					SNew(STextBlock)
+					.TextStyle(FAppStyle::Get(), "SmallText")
+					.Text(LOCTEXT("SaveBasicButton", "Save..."))
+				]
+			];
+	}
+	else if (SelectedJob.IsValid() && SelectedJob->GetConfigMode() == EMoviePipelineConfigMode::Graph)
+	{
+		// Graph asset picker.
+		SubWidget = SNew(SObjectPropertyEntryBox)
+			.AllowedClass(UMovieGraphConfig::StaticClass())
+			.ObjectPath_Lambda([WeakJob]() -> FString
+			{
+				if (WeakJob.IsValid())
+				{
+					if (const UMovieGraphConfig* Graph = WeakJob->GetGraphPreset())
+					{
+						return Graph->GetPathName();
+					}
+				}
+				return FString();
+			})
+			.OnObjectChanged_Lambda([WeakJob, CachedBuilder](const FAssetData& AssetData)
+			{
+				if (WeakJob.IsValid())
+				{
+					WeakJob->Modify();
+					WeakJob->SetGraphPreset(Cast<UMovieGraphConfig>(AssetData.GetAsset()));
+					CachedBuilder->ForceRefreshDetails();
+				}
+			})
+			.AllowClear(true)
+			.DisplayThumbnail(false);
+	}
+	else // Preset
+	{
+		// Preset asset picker.
+		SubWidget = SNew(SObjectPropertyEntryBox)
+			.AllowedClass(UMoviePipelinePrimaryConfig::StaticClass())
+			.ObjectPath_Lambda([WeakJob]() -> FString
+			{
+				if (WeakJob.IsValid())
+				{
+					if (const UMoviePipelinePrimaryConfig* Preset = WeakJob->GetPresetOrigin())
+					{
+						return Preset->GetPathName();
+					}
+				}
+				return FString();
+			})
+			.OnObjectChanged_Lambda([WeakJob, CachedBuilder](const FAssetData& AssetData)
+			{
+				if (WeakJob.IsValid())
+				{
+					WeakJob->Modify();
+					WeakJob->SetPresetOrigin(Cast<UMoviePipelinePrimaryConfig>(AssetData.GetAsset()));
+					CachedBuilder->ForceRefreshDetails();
+				}
+			})
+			.AllowClear(true)
+			.DisplayThumbnail(false);
+	}
+
+	// Single "Configuration" row: buttons on top, sub-widget below.
+	InCategory.AddCustomRow(LOCTEXT("ConfigTypeFilter", "Configuration"))
+		.NameContent()
+		[
+			SNew(STextBlock)
+			.Text(LOCTEXT("ConfigTypeLabel", "Configuration"))
+			.Font(IDetailLayoutBuilder::GetDetailFont())
+		]
+		.ValueContent()
+		[
+			SNew(SVerticalBox)
+			+ SVerticalBox::Slot()
+			.AutoHeight()
+			.Padding(0, 2, 0, 4)
+			[
+				ConfigButtonBox
+			]
+			+ SVerticalBox::Slot()
+			.AutoHeight()
+			[
+				SubWidget
+			]
+		];
+}
+
+void FJobDetailsCustomization_NewQueue::AddBasicConfigOutputTypeProperties(
+	IDetailLayoutBuilder& InDetailBuilder,
+	UMoviePipelineBasicConfig* InActiveConfig,
+	const bool bIsShot)
+{
+	const TArray<UObject*> ActiveConfigObjects = { InActiveConfig };
+
+	IDetailCategoryBuilder& OutputCategory = InDetailBuilder.EditCategory(
+		"Output", LOCTEXT("OutputCategory", "Output"));
+	OptionallyAddShotInlineEditCondition(
+		OutputCategory.AddExternalObjectProperty(ActiveConfigObjects, GET_MEMBER_NAME_CHECKED(UMoviePipelineBasicConfig, OutputDirectory)),
+		InActiveConfig, bIsShot,
+		[InActiveConfig]{ return !!InActiveConfig->bOverride_OutputDirectory; },
+		[InActiveConfig](const bool bNewValue){ InActiveConfig->bOverride_OutputDirectory = bNewValue; });
+	OptionallyAddShotInlineEditCondition(
+		OutputCategory.AddExternalObjectProperty(ActiveConfigObjects, GET_MEMBER_NAME_CHECKED(UMoviePipelineBasicConfig, FileNameFormat)),
+		InActiveConfig, bIsShot,
+		[InActiveConfig]{ return !!InActiveConfig->bOverride_FileNameFormat; },
+		[InActiveConfig](const bool bNewValue){ InActiveConfig->bOverride_FileNameFormat = bNewValue; });
+	OptionallyAddShotInlineEditCondition(
+		OutputCategory.AddExternalObjectProperty(ActiveConfigObjects, GET_MEMBER_NAME_CHECKED(UMoviePipelineBasicConfig, OutputResolution)),
+		InActiveConfig, bIsShot,
+		[InActiveConfig]{ return !!InActiveConfig->bOverride_OutputResolution; },
+		[InActiveConfig](const bool bNewValue){ InActiveConfig->bOverride_OutputResolution = bNewValue; });
+	OptionallyAddShotInlineEditCondition(
+		OutputCategory.AddExternalObjectProperty(ActiveConfigObjects, GET_MEMBER_NAME_CHECKED(UMoviePipelineBasicConfig, CustomStartFrame)),
+		InActiveConfig, bIsShot,
+		[InActiveConfig]{ return !!InActiveConfig->bOverride_CustomStartFrame; },
+		[InActiveConfig](const bool bNewValue){ InActiveConfig->bOverride_CustomStartFrame = bNewValue; });
+	OptionallyAddShotInlineEditCondition(
+		OutputCategory.AddExternalObjectProperty(ActiveConfigObjects, GET_MEMBER_NAME_CHECKED(UMoviePipelineBasicConfig, CustomEndFrame)),
+		InActiveConfig, bIsShot,
+		[InActiveConfig]{ return !!InActiveConfig->bOverride_CustomEndFrame; },
+		[InActiveConfig](const bool bNewValue){ InActiveConfig->bOverride_CustomEndFrame = bNewValue; });
+
+	// Discover ALL output node classes that opt in to Basic config via GetBasicConfigShortDisplayName()
+	auto IsValidBasicConfigOutputType = [](const UClass* Class) -> bool
+	{
+		if (Class->HasAnyClassFlags(CLASS_Abstract | CLASS_Deprecated | CLASS_Hidden | CLASS_HideDropDown))
+		{
+			return false;
+		}
+		const UMovieGraphFileOutputNode* CDO = Cast<UMovieGraphFileOutputNode>(Class->GetDefaultObject());
+		return CDO && !CDO->GetBasicConfigShortDisplayName().IsEmpty();
+	};
+
+	auto SortByDisplayName = [](const UClass& A, const UClass& B)
+	{
+		const UMovieGraphFileOutputNode* CDOA = CastChecked<UMovieGraphFileOutputNode>(A.GetDefaultObject());
+		const UMovieGraphFileOutputNode* CDOB = CastChecked<UMovieGraphFileOutputNode>(B.GetDefaultObject());
+		return CDOA->GetBasicConfigShortDisplayName().CompareTo(CDOB->GetBasicConfigShortDisplayName()) < 0;
+	};
+
+	TArray<UClass*> AllOutputNodeClasses;
+	GetDerivedClasses(UMovieGraphFileOutputNode::StaticClass(), AllOutputNodeClasses, /*bRecursive=*/ true);
+	AllOutputNodeClasses.RemoveAll([&IsValidBasicConfigOutputType](const UClass* Class) { return !IsValidBasicConfigOutputType(Class); });
+	AllOutputNodeClasses.Sort(SortByDisplayName);
+
+	// For shots: output type buttons are enabled only when the override flag is active.
+	const TAttribute<bool> OutputTypeEnabled = bIsShot
+		? TAttribute<bool>::CreateLambda([InActiveConfig]() { return !!InActiveConfig->bOverride_EnabledOutputTypes; })
+		: TAttribute<bool>(true);
+
+	// Create buttons for ALL output types; bind visibility to editor settings so
+	// the gear menu can show/hide them without rebuilding the details panel.
+	const TSharedRef<SWrapBox> OutputButtonBox = SNew(SWrapBox).UseAllottedSize(true);
+	for (const UClass* OutputClass : AllOutputNodeClasses)
+	{
+		TSoftClassPtr<UMovieGraphFileOutputNode> SoftClassPtr(OutputClass);
+		FSoftClassPath ClassPath(OutputClass);
+		const UMovieGraphFileOutputNode* OutputNodeCDO = CastChecked<UMovieGraphFileOutputNode>(OutputClass->GetDefaultObject());
+		FText ButtonLabel = OutputNodeCDO->GetBasicConfigShortDisplayName();
+
+		OutputButtonBox->AddSlot()
+			.Padding(0, 0, 5, 2)
+			[
+				SNew(SCheckBox)
+				.Visibility_Lambda([ClassPath]()
+				{
+					const UMovieRenderGraphEditorSettings* Settings = GetDefault<UMovieRenderGraphEditorSettings>();
+					return Settings->VisibleBasicConfigOutputTypes.Contains(ClassPath)
+						? EVisibility::Visible : EVisibility::Collapsed;
+				})
+				.Style(FAppStyle::Get(), "DetailsView.SectionButton")
+				.IsEnabled(OutputTypeEnabled)
+				.IsChecked_Lambda([InActiveConfig, SoftClassPtr]()
+				{
+					return InActiveConfig->EnabledOutputTypes.Contains(SoftClassPtr)
+						? ECheckBoxState::Checked : ECheckBoxState::Unchecked;
+				})
+				.OnCheckStateChanged_Lambda([InActiveConfig, SoftClassPtr](ECheckBoxState NewState)
+				{
+					if (NewState == ECheckBoxState::Unchecked && InActiveConfig->EnabledOutputTypes.Num() <= 1)
+					{
+						return;
+					}
+					InActiveConfig->Modify();
+					if (NewState == ECheckBoxState::Checked)
+					{
+						InActiveConfig->EnabledOutputTypes.AddUnique(SoftClassPtr);
+					}
+					else
+					{
+						InActiveConfig->EnabledOutputTypes.Remove(SoftClassPtr);
+					}
+				})
+				[
+					SNew(STextBlock)
+					.TextStyle(FAppStyle::Get(), "SmallText")
+					.Text(ButtonLabel)
+				]
+			];
+	}
+
+	// Gear menu for controlling which output types are visible and burn-in visibility.
+	// All changes take effect immediately via visibility bindings.
+	const TSharedRef<SComboButton> OutputTypeGearMenu = SNew(SComboButton)
+		.HasDownArrow(false)
+		.ButtonStyle(FAppStyle::Get(), "SimpleButton")
+		.OnGetMenuContent_Lambda([InActiveConfig, AllOutputNodeClasses]() -> TSharedRef<SWidget>
+		{
+			FMenuBuilder MenuBuilder(/*bInShouldCloseWindowAfterMenuSelection=*/ false, nullptr);
+			UMovieRenderGraphEditorSettings* MutableSettings = GetMutableDefault<UMovieRenderGraphEditorSettings>();
+
+			MenuBuilder.BeginSection(NAME_None, LOCTEXT("BasicOutputTypesSection", "Basic Output Types"));
+			for (const UClass* OutputClass : AllOutputNodeClasses)
+			{
+				FSoftClassPath ClassPath(OutputClass);
+				const UMovieGraphFileOutputNode* OutputNodeCDO = CastChecked<UMovieGraphFileOutputNode>(OutputClass->GetDefaultObject());
+				FText Label = OutputNodeCDO->GetBasicConfigShortDisplayName();
+
+				MenuBuilder.AddMenuEntry(
+					Label,
+					FText::GetEmpty(),
+					FSlateIcon(),
+					FUIAction(
+						FExecuteAction::CreateLambda([MutableSettings, ClassPath, InActiveConfig]()
+						{
+							if (MutableSettings->VisibleBasicConfigOutputTypes.Contains(ClassPath))
+							{
+								// Don't allow hiding the last enabled output type
+								const TSoftClassPtr<UMovieGraphFileOutputNode> SoftPtr(ClassPath);
+								if (InActiveConfig->EnabledOutputTypes.Contains(SoftPtr) && InActiveConfig->EnabledOutputTypes.Num() <= 1)
+								{
+									return;
+								}
+								
+								MutableSettings->VisibleBasicConfigOutputTypes.Remove(ClassPath);
+								
+								// Remove from EnabledOutputTypes to prevent phantom graph nodes
+								InActiveConfig->Modify();
+								InActiveConfig->EnabledOutputTypes.Remove(SoftPtr);
+							}
+							else
+							{
+								MutableSettings->VisibleBasicConfigOutputTypes.AddUnique(ClassPath);
+							}
+							
+							MutableSettings->SaveConfig();
+						}),
+						FCanExecuteAction(),
+						FIsActionChecked::CreateLambda([MutableSettings, ClassPath]()
+						{
+							return MutableSettings->VisibleBasicConfigOutputTypes.Contains(ClassPath);
+						})
+					),
+					NAME_None,
+					EUserInterfaceActionType::ToggleButton
+				);
+			}
+			MenuBuilder.EndSection();
+
+			MenuBuilder.AddSeparator();
+
+			MenuBuilder.BeginSection(NAME_None);
+			MenuBuilder.AddMenuEntry(
+				LOCTEXT("BurnInClassMenu", "Burn In Class"),
+				LOCTEXT("BurnInClassMenuTooltip", "Show the Burn In Class property in Basic Config"),
+				FSlateIcon(),
+				FUIAction(
+					FExecuteAction::CreateLambda([MutableSettings]()
+					{
+						MutableSettings->bShowBurnInForBasicConfig = !MutableSettings->bShowBurnInForBasicConfig;
+						MutableSettings->SaveConfig();
+					}),
+					FCanExecuteAction(),
+					FIsActionChecked::CreateLambda([MutableSettings]()
+					{
+						return MutableSettings->bShowBurnInForBasicConfig;
+					})
+				),
+				NAME_None,
+				EUserInterfaceActionType::ToggleButton
+			);
+			MenuBuilder.EndSection();
+
+			return MenuBuilder.MakeWidget();
+		})
+		.ButtonContent()
+		[
+			SNew(SImage)
+			.Image(FAppStyle::Get().GetBrush("Icons.Settings"))
+			.DesiredSizeOverride(FVector2D(16.f, 16.f))
+			.ColorAndOpacity(FSlateColor::UseForeground())
+		];
+
+	OutputCategory.AddCustomRow(LOCTEXT("OutputTypeFilter", "Output Type"))
+		.NameContent()
+		[
+			SNew(SHorizontalBox)
+			+ SHorizontalBox::Slot()
+			.AutoWidth()
+			.VAlign(VAlign_Center)
+			.Padding(-2, 0, 2, 0)
+			[
+				// Override toggle — visible only in shot mode
+				SNew(SCheckBox)
+				.Visibility(bIsShot ? EVisibility::Visible : EVisibility::Collapsed)
+				.IsChecked_Lambda([InActiveConfig]()
+				{
+					return InActiveConfig->bOverride_EnabledOutputTypes
+						? ECheckBoxState::Checked : ECheckBoxState::Unchecked;
+				})
+				.OnCheckStateChanged_Lambda([InActiveConfig](const ECheckBoxState NewState)
+				{
+					InActiveConfig->Modify();
+					InActiveConfig->bOverride_EnabledOutputTypes = (NewState == ECheckBoxState::Checked) ? 1 : 0;
+				})
+			]
+			+ SHorizontalBox::Slot()
+			.AutoWidth()
+			.VAlign(VAlign_Center)
+			[
+				SNew(STextBlock)
+				.Text(LOCTEXT("OutputTypeLabel", "Output Type"))
+				.Font(IDetailLayoutBuilder::GetDetailFont())
+				.IsEnabled(OutputTypeEnabled)
+			]
+		]
+		.ValueContent()
+		.HAlign(HAlign_Fill)
+		[
+			SNew(SHorizontalBox)
+			+ SHorizontalBox::Slot()
+			.FillWidth(1.f)
+			[
+				OutputButtonBox
+			]
+			+ SHorizontalBox::Slot()
+			.AutoWidth()
+			.VAlign(VAlign_Center)
+			.Padding(4, 0, 0, 0)
+			[
+				OutputTypeGearMenu
+			]
+		];
+
+	// BurnInClass row — always created, visibility bound to editor settings
+	if (IDetailPropertyRow* BurnInRow = OutputCategory.AddExternalObjectProperty(ActiveConfigObjects,
+		GET_MEMBER_NAME_CHECKED(UMoviePipelineBasicConfig, BurnInClass)))
+	{
+		BurnInRow->Visibility(TAttribute<EVisibility>::CreateLambda([]()
+		{
+			const UMovieRenderGraphEditorSettings* Settings = GetDefault<UMovieRenderGraphEditorSettings>();
+			return Settings->bShowBurnInForBasicConfig ? EVisibility::Visible : EVisibility::Collapsed;
+		}));
+		OptionallyAddShotInlineEditCondition(BurnInRow, InActiveConfig, bIsShot,
+			[InActiveConfig]{ return !!InActiveConfig->bOverride_BurnInClass; },
+			[InActiveConfig](const bool bInNewValue){ InActiveConfig->bOverride_BurnInClass = bInNewValue; });
+	}
+}
+
+void FJobDetailsCustomization_NewQueue::AddBasicConfigRenderingProperties(
+	IDetailLayoutBuilder& InDetailBuilder,
+	UMoviePipelineBasicConfig* InActiveConfig,
+	const bool bIsShot)
+{
+	const TArray<UObject*> ActiveConfigObjects = { InActiveConfig };
+
+	IDetailCategoryBuilder& RenderingCategory = InDetailBuilder.EditCategory(
+		"Rendering", LOCTEXT("RenderingCategory", "Rendering"));
+	RenderingCategory.SetShowAdvanced(false);
+
+	// Renderer toggle row -- for shots, a single row-level override checkbox controls
+	// both renderer flags together (renderer choice is a single setting, not two independent ones).
+	const FToggleButtonDef RendererButtons[] = {
+		{ LOCTEXT("DeferredToggle", "Deferred"), &InActiveConfig->bUseDeferredRenderer,
+			[InActiveConfig]{ return !!InActiveConfig->bUsePathTracedRenderer; } },
+		{ LOCTEXT("PathTracedToggle", "Path Traced"), &InActiveConfig->bUsePathTracedRenderer,
+			[InActiveConfig]{ return !!InActiveConfig->bUseDeferredRenderer; } },
+	};
+	AddToggleButtonRow(RenderingCategory, InActiveConfig,
+		LOCTEXT("RendererFilter", "Renderer"),
+		LOCTEXT("RendererLabel", "Renderer"),
+		RendererButtons,
+		// Row override: enabled in shot mode, controlling both renderer flags together.
+		bIsShot
+			? TFunction<bool()>([InActiveConfig]{ return !!InActiveConfig->bOverride_bUseDeferredRenderer; })
+			: TFunction<bool()>{},
+		bIsShot
+			? TFunction<void(bool)>([InActiveConfig](const bool bNewValue)
+				{
+					InActiveConfig->bOverride_bUseDeferredRenderer = bNewValue;
+					InActiveConfig->bOverride_bUsePathTracedRenderer = bNewValue;
+				})
+			: TFunction<void(bool)>{});
+
+	// Properties that apply to all renderers
+	OptionallyAddShotInlineEditCondition(
+		RenderingCategory.AddExternalObjectProperty(ActiveConfigObjects,
+			GET_MEMBER_NAME_CHECKED(UMoviePipelineBasicConfig, NumWarmUpFrames)),
+		InActiveConfig, bIsShot,
+		[InActiveConfig]{ return !!InActiveConfig->bOverride_NumWarmUpFrames; },
+		[InActiveConfig](bool v){ InActiveConfig->bOverride_NumWarmUpFrames = v; });
+	OptionallyAddShotInlineEditCondition(
+		RenderingCategory.AddExternalObjectProperty(ActiveConfigObjects,
+			GET_MEMBER_NAME_CHECKED(UMoviePipelineBasicConfig, TemporalSampleCount)),
+		InActiveConfig, bIsShot,
+		[InActiveConfig]{ return !!InActiveConfig->bOverride_TemporalSampleCount; },
+		[InActiveConfig](bool v){ InActiveConfig->bOverride_TemporalSampleCount = v; });
+
+	// Always create both renderer sub-groups; bind visibility so toggling is immediate without a full panel refresh.
+	const TAttribute<EVisibility> DeferredVisibility = TAttribute<EVisibility>::CreateLambda([InActiveConfig]()
+	{
+		return InActiveConfig->bUseDeferredRenderer ? EVisibility::Visible : EVisibility::Collapsed;
+	});
+	const TAttribute<EVisibility> PathTracedVisibility = TAttribute<EVisibility>::CreateLambda([InActiveConfig]()
+	{
+		return InActiveConfig->bUsePathTracedRenderer ? EVisibility::Visible : EVisibility::Collapsed;
+	});
+
+	{
+		IDetailGroup& DeferredGroup = RenderingCategory.AddGroup("DeferredRenderer",
+			LOCTEXT("DeferredRendererGroup", "Deferred Renderer"), /*bForAdvanced=*/ false, /*bStartExpanded=*/ true);
+		DeferredGroup.HeaderRow()
+			.Visibility(DeferredVisibility)
+			.NameContent()
+			[
+				SNew(STextBlock)
+				.Font(IDetailLayoutBuilder::GetDetailFont())
+				.Text(LOCTEXT("DeferredRendererGroup", "Deferred Renderer"))
+			];
+		{
+			IDetailPropertyRow& DeferredSpatialRow = DeferredGroup.AddExternalObjectProperty(ActiveConfigObjects,
+				GET_MEMBER_NAME_CHECKED(UMoviePipelineBasicConfig, DeferredSpatialSampleCount), EPropertyLocation::Default, FAddPropertyParams());
+			DeferredSpatialRow.Visibility(DeferredVisibility);
+			OptionallyAddShotInlineEditCondition(&DeferredSpatialRow, InActiveConfig, bIsShot,
+				[InActiveConfig]{ return !!InActiveConfig->bOverride_DeferredSpatialSampleCount; },
+				[InActiveConfig](const bool bNewValue){ InActiveConfig->bOverride_DeferredSpatialSampleCount = bNewValue; });
+		}
+	}
+
+	{
+		IDetailGroup& PathTracedGroup = RenderingCategory.AddGroup("PathTracedRenderer",
+			LOCTEXT("PathTracedRendererGroup", "Path Traced Renderer"), /*bForAdvanced=*/ false, /*bStartExpanded=*/ true);
+		PathTracedGroup.HeaderRow()
+			.Visibility(PathTracedVisibility)
+			.NameContent()
+			[
+				SNew(STextBlock)
+				.Font(IDetailLayoutBuilder::GetDetailFont())
+				.Text(LOCTEXT("PathTracedRendererGroup", "Path Traced Renderer"))
+			];
+		{
+			IDetailPropertyRow& PathTracerSpatialRow = PathTracedGroup.AddExternalObjectProperty(ActiveConfigObjects,
+				GET_MEMBER_NAME_CHECKED(UMoviePipelineBasicConfig, PathTracedSpatialSampleCount), EPropertyLocation::Default, FAddPropertyParams());
+			PathTracerSpatialRow.Visibility(PathTracedVisibility);
+			OptionallyAddShotInlineEditCondition(&PathTracerSpatialRow, InActiveConfig, bIsShot,
+				[InActiveConfig]{ return !!InActiveConfig->bOverride_PathTracedSpatialSampleCount; },
+				[InActiveConfig](const bool bNewValue){ InActiveConfig->bOverride_PathTracedSpatialSampleCount = bNewValue; });
+		}
+		{
+			IDetailPropertyRow& PathTracerDenoiserRow = PathTracedGroup.AddExternalObjectProperty(ActiveConfigObjects,
+				GET_MEMBER_NAME_CHECKED(UMoviePipelineBasicConfig, PathTracedDenoiserType), EPropertyLocation::Default, FAddPropertyParams());
+			PathTracerDenoiserRow.Visibility(PathTracedVisibility);
+			OptionallyAddShotInlineEditCondition(&PathTracerDenoiserRow, InActiveConfig, bIsShot,
+				[InActiveConfig]{ return !!InActiveConfig->bOverride_PathTracedDenoiserType; },
+				[InActiveConfig](const bool bNewValue){ InActiveConfig->bOverride_PathTracedDenoiserType = bNewValue; });
+		}
+	}
+}
+
+#undef LOCTEXT_NAMESPACE
diff --git a/Engine/Plugins/MovieScene/MovieRenderPipeline/Source/MovieRenderPipelineEditor/Private/Customizations/JobCustomization_NewQueue.h b/Engine/Plugins/MovieScene/MovieRenderPipeline/Source/MovieRenderPipelineEditor/Private/Customizations/JobCustomization_NewQueue.h
new file mode 100644
--- /dev/null
+++ b/Engine/Plugins/MovieScene/MovieRenderPipeline/Source/MovieRenderPipelineEditor/Private/Customizations/JobCustomization_NewQueue.h
@@ -0,0 +1,61 @@
+// Copyright Epic Games, Inc. All Rights Reserved.
+
+#pragma once
+
+#include "IDetailCustomization.h"
+#include "UObject/WeakObjectPtr.h"
+
+class FObjectPostSaveContext;
+class IDetailCategoryBuilder;
+class IDetailLayoutBuilder;
+class UMovieGraphConfig;
+class UMoviePipelineBasicConfig;
+class UMoviePipelineExecutorJob;
+class UMoviePipelineExecutorShot;
+class UPackage;
+enum class EMoviePipelineConfigMode : uint8;
+
+/**
+ * Detail customization for jobs and shots in the new queue panel.
+ *
+ * Surfaces the Basic configuration mode UI (segmented control, output/rendering categories,
+ * shot override toggles) in addition to the existing Graph/Preset variable assignments.
+ */
+class FJobDetailsCustomization_NewQueue : public IDetailCustomization
+{
+public:
+	static TSharedRef<IDetailCustomization> MakeInstance();
+	virtual ~FJobDetailsCustomization_NewQueue() override = default;
+
+	virtual void PendingDelete() override;
+
+protected:
+	//~ Begin IDetailCustomization interface
+	virtual void CustomizeDetails(const TSharedPtr<IDetailLayoutBuilder>& InDetailBuilder) override;
+	virtual void CustomizeDetails(IDetailLayoutBuilder& InDetailBuilder) override;
+	//~ End IDetailCustomization interface
+
+private:
+	void RefreshLayout(const FString&, UPackage*, FObjectPostSaveContext) const;
+	void RefreshLayout(UMoviePipelineExecutorShot*, UMovieGraphConfig*) const;
+	void RefreshLayout(UMoviePipelineExecutorJob*, UMovieGraphConfig*) const;
+	void RefreshLayout(UMoviePipelineExecutorJob*, EMoviePipelineConfigMode) const;
+
+	/** Adds the Configuration mode row (segmented control + mode-dependent sub-widget) for primary jobs. */
+	void AddConfigurationModeRow(IDetailCategoryBuilder& InCategory);
+
+	/** Adds the Output category with output type buttons, gear menu, and per-property override toggles. */
+	void AddBasicConfigOutputTypeProperties(IDetailLayoutBuilder& InDetailBuilder, UMoviePipelineBasicConfig* InActiveConfig, bool bIsShot);
+
+	/** Adds the Rendering category with renderer toggles and per-renderer sub-groups. */
+	void AddBasicConfigRenderingProperties(IDetailLayoutBuilder& InDetailBuilder, UMoviePipelineBasicConfig* InActiveConfig, bool bIsShot);
+
+	/** The details builder associated with the customization. */
+	IDetailLayoutBuilder* DetailBuilder = nullptr;
+
+	/** The primary job that's selected in the UI (can be null). */
+	TWeakObjectPtr<UMoviePipelineExecutorJob> SelectedJob;
+
+	/** The shot that's selected in the UI (can be null if a job is selected). */
+	TWeakObjectPtr<UMoviePipelineExecutorShot> SelectedShot;
+};
diff --git a/Engine/Plugins/MovieScene/MovieRenderPipeline/Source/MovieRenderPipelineEditor/Private/Widgets/SMoviePipelineAfterRenderWidget.cpp b/Engine/Plugins/MovieScene/MovieRenderPipeline/Source/MovieRenderPipelineEditor/Private/Widgets/SMoviePipelineAfterRenderWidget.cpp
new file mode 100644
--- /dev/null
+++ b/Engine/Plugins/MovieScene/MovieRenderPipeline/Source/MovieRenderPipelineEditor/Private/Widgets/SMoviePipelineAfterRenderWidget.cpp
@@ -0,0 +1,77 @@
+// Copyright Epic Games, Inc. All Rights Reserved.
+
+#include "Widgets/SMoviePipelineAfterRenderWidget.h"
+#include "MoviePipelineQueue.h"
+
+#include "PropertyEditorModule.h"
+#include "ISinglePropertyView.h"
+#include "ISettingsModule.h"
+#include "Modules/ModuleManager.h"
+
+#include "Widgets/SBoxPanel.h"
+#include "Widgets/Text/STextBlock.h"
+#include "Widgets/Input/SButton.h"
+#include "Widgets/Images/SImage.h"
+#include "Styling/AppStyle.h"
+
+#define LOCTEXT_NAMESPACE "SMoviePipelineAfterRenderWidget"
+
+void SMoviePipelineAfterRenderWidget::Construct(const FArguments& InArgs)
+{
+	FPropertyEditorModule& PropertyEditorModule = FModuleManager::Get().LoadModuleChecked<FPropertyEditorModule>("PropertyEditor");
+
+	UMoviePipelineQueueProjectSettings* QueueProjectSettings = GetMutableDefault<UMoviePipelineQueueProjectSettings>();
+
+	FSinglePropertyParams PostRenderActionTypeParams;
+	PostRenderActionTypeParams.NamePlacement = EPropertyNamePlacement::Hidden;
+	PostRenderActionTypeParams.bHideResetToDefault = true;
+
+	PostRenderBehaviorWidget = PropertyEditorModule.CreateSingleProperty(
+		QueueProjectSettings, GET_MEMBER_NAME_CHECKED(UMoviePipelineQueueProjectSettings, PostRenderActionType), PostRenderActionTypeParams);
+
+	ChildSlot
+	[
+		SNew(SHorizontalBox)
+
+		+ SHorizontalBox::Slot()
+		.VAlign(VAlign_Center)
+		.AutoWidth()
+		.Padding(4, 0, 4, 0)
+		[
+			SNew(STextBlock)
+			.Text(LOCTEXT("AfterRender_Text", "After Render:"))
+		]
+
+		+ SHorizontalBox::Slot()
+		.VAlign(VAlign_Center)
+		.AutoWidth()
+		[
+			PostRenderBehaviorWidget.ToSharedRef()
+		]
+
+		+ SHorizontalBox::Slot()
+		.VAlign(VAlign_Center)
+		.AutoWidth()
+		[
+			SNew(SButton)
+			.ButtonStyle(FAppStyle::Get(), TEXT("SimpleButton"))
+			.ContentPadding(0)
+			.OnClicked_Lambda([]()
+			{
+				FModuleManager::LoadModuleChecked<ISettingsModule>("Settings")
+					.ShowViewer("Editor", "Plugins", "MovieRenderGraphEditorSettings");
+
+				return FReply::Handled();
+			})
+			.ToolTipText(LOCTEXT("OpenPlaybackPrefs_Tooltip", "Open up the preferences which dictate how media is opened after a render completes."))
+			[
+				SNew(SImage)
+				.ColorAndOpacity(FSlateColor::UseForeground())
+				.Image(FAppStyle::Get().GetBrush("EditorPreferences.TabIcon"))
+				.DesiredSizeOverride(FVector2D(16.f, 16.f))
+			]
+		]
+	];
+}
+
+#undef LOCTEXT_NAMESPACE
diff --git a/Engine/Plugins/MovieScene/MovieRenderPipeline/Source/MovieRenderPipelineEditor/Private/Widgets/SMoviePipelineAfterRenderWidget.h b/Engine/Plugins/MovieScene/MovieRenderPipeline/Source/MovieRenderPipelineEditor/Private/Widgets/SMoviePipelineAfterRenderWidget.h
new file mode 100644
--- /dev/null
+++ b/Engine/Plugins/MovieScene/MovieRenderPipeline/Source/MovieRenderPipelineEditor/Private/Widgets/SMoviePipelineAfterRenderWidget.h
@@ -0,0 +1,24 @@
+// Copyright Epic Games, Inc. All Rights Reserved.
+
+#pragma once
+
+#include "Widgets/DeclarativeSyntaxSupport.h"
+#include "Widgets/SCompoundWidget.h"
+
+class ISinglePropertyView;
+
+/**
+ * Compact toolbar widget that shows the "After Render" behavior dropdown
+ * and a settings button to open the post-render preferences.
+ */
+class SMoviePipelineAfterRenderWidget : public SCompoundWidget
+{
+public:
+	SLATE_BEGIN_ARGS(SMoviePipelineAfterRenderWidget) {}
+	SLATE_END_ARGS()
+
+	void Construct(const FArguments& InArgs);
+
+private:
+	TSharedPtr<ISinglePropertyView> PostRenderBehaviorWidget;
+};
diff --git a/Engine/Plugins/MovieScene/MovieRenderPipeline/Source/MovieRenderPipelineEditor/Private/Widgets/SMoviePipelineNewQueuePanel.cpp b/Engine/Plugins/MovieScene/MovieRenderPipeline/Source/MovieRenderPipelineEditor/Private/Widgets/SMoviePipelineNewQueuePanel.cpp
--- a/Engine/Plugins/MovieScene/MovieRenderPipeline/Source/MovieRenderPipelineEditor/Private/Widgets/SMoviePipelineNewQueuePanel.cpp
+++ b/Engine/Plugins/MovieScene/MovieRenderPipeline/Source/MovieRenderPipelineEditor/Private/Widgets/SMoviePipelineNewQueuePanel.cpp
@@ -1,47 +1,47 @@
 // Copyright Epic Games, Inc. All Rights Reserved.
 
 #include "Widgets/SMoviePipelineNewQueuePanel.h"
-#include "MoviePipelineQueuePanelUtils.h"
+
 #include "Customizations/Graph/MovieGraphNamedResolutionCustomization.h"
-#include "Customizations/JobCustomization.h"
-#include "Widgets/MoviePipelineWidgetConstants.h"
-#include "SMoviePipelineQueueEditor.h"
-#include "SMoviePipelineConfigPanel.h"
+#include "Customizations/JobCustomization_NewQueue.h"
+#include "Graph/MovieGraphAssetToolkit.h"
+#include "Graph/MovieGraphConfig.h"
+#include "Graph/MovieRenderGraphEditorSettings.h"
 #include "MovieRenderPipelineEditorUtils.h"
 #include "MovieRenderPipelineSettings.h"
 #include "MovieRenderPipelineStyle.h"
-#include "MoviePipelineQueue.h"
 #include "MoviePipelinePrimaryConfig.h"
+#include "MoviePipelineQueue.h"
+#include "MoviePipelineQueuePanelUtils.h"
 #include "MoviePipelineQueueSubsystem.h"
 #include "MoviePipelineTelemetry.h"
-#include "Graph/MovieGraphConfig.h"
-#include "Graph/MovieGraphAssetToolkit.h"
+#include "SMoviePipelineQueueEditor.h"
+#include "SMoviePipelineConfigPanel.h"
+#include "Widgets/MoviePipelineWidgetConstants.h"
+#include "Widgets/SMoviePipelineAfterRenderWidget.h"
 
-#include "Widgets/SBoxPanel.h"
-#include "Widgets/Input/SComboBox.h"
-#include "Widgets/Layout/SBorder.h"
-#include "Widgets/Text/STextBlock.h"
-#include "Widgets/Input/SButton.h"
-#include "Widgets/SWindow.h"
-#include "Styling/AppStyle.h"
 #include "Framework/Application/SlateApplication.h"
-#include "Widgets/Layout/SSplitter.h"
-#include "Widgets/Layout/SWidgetSwitcher.h"
-#include "Widgets/Images/SImage.h"
 #include "Framework/MultiBox/MultiBoxBuilder.h"
+#include "Styling/AppStyle.h"
 #include "SPositiveActionButton.h"
 #include "Styling/ToolBarStyle.h"
+#include "Widgets/Input/SComboBox.h"
+#include "Widgets/Images/SImage.h"
+#include "Widgets/Layout/SBorder.h"
+#include "Widgets/Layout/SSeparator.h"
+#include "Widgets/Layout/SSplitter.h"
+#include "Widgets/Layout/SWidgetSwitcher.h"
+#include "Widgets/SBoxPanel.h"
+#include "Widgets/SWindow.h"
+#include "Widgets/Text/STextBlock.h"
 
-#include "IContentBrowserSingleton.h"
+#include "AssetRegistry/AssetData.h"
 #include "ContentBrowserModule.h"
-
 #include "Editor.h"
-#include "PropertyEditorModule.h"
+#include "IContentBrowserSingleton.h"
 #include "IDetailsView.h"
-#include "ISettingsModule.h"
-#include "ISinglePropertyView.h"
 #include "Modules/ModuleManager.h"
-#include "AssetRegistry/AssetData.h"
+#include "PropertyEditorModule.h"
 
 #define LOCTEXT_NAMESPACE "SMoviePipelineNewQueuePanel"
 
@@ -50,20 +50,66 @@ void SMoviePipelineNewQueuePanel::Construct(const FArguments& InArgs)
 	FPropertyEditorModule& PropertyEditorModule = FModuleManager::Get().LoadModuleChecked<FPropertyEditorModule>("PropertyEditor");
 
 	FDetailsViewArgs DetailsViewArgs;
-	DetailsViewArgs.bAllowSearch = false;
-	DetailsViewArgs.NameAreaSettings = FDetailsViewArgs::HideNameArea;
+	DetailsViewArgs.bAllowSearch = true;
+	DetailsViewArgs.bLockable = true;
+	DetailsViewArgs.NameAreaSettings = FDetailsViewArgs::ObjectsUseNameArea;
+	DetailsViewArgs.bShowObjectLabel = true;
+	DetailsViewArgs.bCustomNameAreaLocation = true;
+	DetailsViewArgs.bCustomFilterAreaLocation = true;
 	DetailsViewArgs.bHideSelectionTip = true;
 	DetailsViewArgs.ColumnWidth = 0.7f;
 
+	// Extend the built-in gear/settings menu with a "Default Configuration Type" section.
+	const TSharedPtr<FExtender> OptionsExtender = MakeShareable(new FExtender);
+	OptionsExtender->AddMenuExtension(
+		FName("DetailsViewShowOptions"),
+		EExtensionHook::Before,
+		nullptr,
+		FMenuExtensionDelegate::CreateLambda([](FMenuBuilder& MenuBuilder)
+		{
+			UMovieRenderGraphEditorSettings* Settings = GetMutableDefault<UMovieRenderGraphEditorSettings>();
+
+			MenuBuilder.BeginSection(NAME_None, LOCTEXT("DefaultConfigTypeSection", "Default Configuration Type"));
+
+			auto AddConfigTypeEntry = [&](const FText& Label, EMoviePipelineDefaultConfigType Type)
+			{
+				MenuBuilder.AddMenuEntry(
+					Label, FText::GetEmpty(), FSlateIcon(),
+					FUIAction(
+						FExecuteAction::CreateLambda([Settings, Type]()
+						{
+							Settings->DefaultConfigType = Type;
+							Settings->SaveConfig();
+						}),
+						FCanExecuteAction(),
+						FIsActionChecked::CreateLambda([Settings, Type]()
+						{
+							return Settings->DefaultConfigType == Type;
+						})
+					),
+					NAME_None,
+					EUserInterfaceActionType::RadioButton
+				);
+			};
+
+			AddConfigTypeEntry(LOCTEXT("ConfigType_Basic", "Basic"), EMoviePipelineDefaultConfigType::Basic);
+			AddConfigTypeEntry(LOCTEXT("ConfigType_Graph", "Movie Render Graph"), EMoviePipelineDefaultConfigType::MovieRenderGraph);
+			AddConfigTypeEntry(LOCTEXT("ConfigType_Preset", "Legacy Preset"), EMoviePipelineDefaultConfigType::LegacyPreset);
+
+			MenuBuilder.EndSection();
+		})
+	);
+	DetailsViewArgs.OptionsExtender = OptionsExtender;
+
 	JobDetailsPanelWidget = PropertyEditorModule.CreateDetailView(DetailsViewArgs);
 
 	JobDetailsPanelWidget->RegisterInstancedCustomPropertyLayout(
 		UMoviePipelineExecutorJob::StaticClass(),
-		FOnGetDetailCustomizationInstance::CreateStatic(&FJobDetailsCustomization::MakeInstance));
+		FOnGetDetailCustomizationInstance::CreateStatic(&FJobDetailsCustomization_NewQueue::MakeInstance));
 
 	JobDetailsPanelWidget->RegisterInstancedCustomPropertyLayout(
 		UMoviePipelineExecutorShot::StaticClass(),
-		FOnGetDetailCustomizationInstance::CreateStatic(&FJobDetailsCustomization::MakeInstance));
+		FOnGetDetailCustomizationInstance::CreateStatic(&FJobDetailsCustomization_NewQueue::MakeInstance));
 
 	JobDetailsPanelWidget->RegisterInstancedCustomPropertyTypeLayout(
 		FMovieGraphNamedResolution::StaticStruct()->GetFName(),
@@ -88,19 +134,6 @@ void SMoviePipelineNewQueuePanel::Construct(const FArguments& InArgs)
 		}
 	}
 
-	// "After Render" behavior dropdown
-	TSharedPtr<ISinglePropertyView> PostRenderBehaviorWidget;
-	{
-		UMoviePipelineQueueProjectSettings* QueueProjectSettings = GetMutableDefault<UMoviePipelineQueueProjectSettings>();
-
-		FSinglePropertyParams PostRenderActionTypeParams;
-		PostRenderActionTypeParams.NamePlacement = EPropertyNamePlacement::Hidden;
-		PostRenderActionTypeParams.bHideResetToDefault = true;
-
-		PostRenderBehaviorWidget = PropertyEditorModule.CreateSingleProperty(
-			QueueProjectSettings, GET_MEMBER_NAME_CHECKED(UMoviePipelineQueueProjectSettings, PostRenderActionType), PostRenderActionTypeParams);
-	}
-
 	ChildSlot
 	[
 		SNew(SVerticalBox)
@@ -166,6 +199,13 @@ void SMoviePipelineNewQueuePanel::Construct(const FArguments& InArgs)
 			]
 		]
 
+		+ SVerticalBox::Slot()
+		.AutoHeight()
+		[
+			SNew(SSeparator)
+			.Orientation(EOrientation::Orient_Horizontal)
+		]
+
 		// Main queue body
 		+ SVerticalBox::Slot()
 		.FillHeight(1.0f)
@@ -190,7 +230,23 @@ void SMoviePipelineNewQueuePanel::Construct(const FArguments& InArgs)
 					.IsEnabled(this, &SMoviePipelineNewQueuePanel::IsDetailsViewEnabled)
 					+ SWidgetSwitcher::Slot()
 					[
-						JobDetailsPanelWidget.ToSharedRef()
+						SNew(SVerticalBox)
+						+ SVerticalBox::Slot()
+						.AutoHeight()
+						.Padding(6.0f, 2.0f, 2.0f, 0.0f)
+						[
+							JobDetailsPanelWidget->GetNameAreaWidget().ToSharedRef()
+						]
+						+ SVerticalBox::Slot()
+						.AutoHeight()
+						[
+							JobDetailsPanelWidget->GetFilterAreaWidget().ToSharedRef()
+						]
+						+ SVerticalBox::Slot()
+						.FillHeight(1.0f)
+						[
+							JobDetailsPanelWidget.ToSharedRef()
+						]
 					]
 					+ SWidgetSwitcher::Slot()
 					.Padding(2.0f, 24.0f, 2.0f, 2.0f)
@@ -207,68 +263,6 @@ void SMoviePipelineNewQueuePanel::Construct(const FArguments& InArgs)
 			]
 		]
 
-		// Footer
-		+ SVerticalBox::Slot()
-		.AutoHeight()
-		[
-			SNew(SBorder)
-			.BorderImage(FAppStyle::GetBrush("DetailsView.CategoryTop"))
-			.BorderBackgroundColor(FLinearColor(.6, .6, .6, 1.0f))
-			.Padding(FMargin(0, 2, 0, 2))
-			[
-				SNew(SHorizontalBox)
-
-				+ SHorizontalBox::Slot()
-				.VAlign(VAlign_Center)
-				.HAlign(HAlign_Left)
-				.Padding(5, 0, 0, 0)
-				.AutoWidth()
-				[
-					SNew(STextBlock)
-					.Text(LOCTEXT("AfterRender_Text", "After Render:"))
-				]
-
-				+ SHorizontalBox::Slot()
-				.VAlign(VAlign_Center)
-				.HAlign(HAlign_Left)
-				.AutoWidth()
-				[
-					PostRenderBehaviorWidget.ToSharedRef()
-				]
-
-				+ SHorizontalBox::Slot()
-				.VAlign(VAlign_Center)
-				.HAlign(HAlign_Left)
-				.AutoWidth()
-				[
-					SNew(SButton)
-					.ButtonStyle(FAppStyle::Get(), TEXT("SimpleButton"))
-					.ContentPadding(0)
-					.OnClicked_Lambda([]()
-					{
-						FModuleManager::LoadModuleChecked<ISettingsModule>("Settings")
-							.ShowViewer("Editor", "Plugins", "MovieRenderGraphEditorSettings");
-
-						return FReply::Handled();
-					})
-					.ToolTipText(LOCTEXT("OpenPlaybackPrefs_Tooltip", "Open up the preferences which dictate how media is opened after a render completes."))
-					[
-						SNew(SImage)
-						.ColorAndOpacity(FSlateColor::UseForeground())
-						.Image(FAppStyle::Get().GetBrush("EditorPreferences.TabIcon"))
-						.DesiredSizeOverride(FVector2D(16.f, 16.f))
-					]
-				]
-
-				+ SHorizontalBox::Slot()
-				.VAlign(VAlign_Fill)
-				.HAlign(HAlign_Left)
-				.FillWidth(1.f)
-				[
-					SNullWidget::NullWidget
-				]
-			]
-		]
 	];
 }
 
@@ -333,6 +327,15 @@ TSharedRef<SWidget> SMoviePipelineNewQueuePanel::MakeToolbar()
 				.IsEnabled(this, &SMoviePipelineNewQueuePanel::IsRenderEnabled)
 				.OnGetMenuContent(this, &SMoviePipelineNewQueuePanel::OnGenerateRenderMenu)
 			]);
+
+	}
+	ToolbarBuilder.EndSection();
+
+	// Section 3: After Render
+	ToolbarBuilder.BeginSection("AfterRender");
+	{
+		ToolbarBuilder.AddWidget(
+			SNew(SMoviePipelineAfterRenderWidget));
 	}
 	ToolbarBuilder.EndSection();
 
diff --git a/Engine/Plugins/MovieScene/MovieRenderPipeline/Source/MovieRenderPipelineEditor/Private/Widgets/SMoviePipelineQueueEditor.cpp b/Engine/Plugins/MovieScene/MovieRenderPipeline/Source/MovieRenderPipelineEditor/Private/Widgets/SMoviePipelineQueueEditor.cpp
--- a/Engine/Plugins/MovieScene/MovieRenderPipeline/Source/MovieRenderPipelineEditor/Private/Widgets/SMoviePipelineQueueEditor.cpp
+++ b/Engine/Plugins/MovieScene/MovieRenderPipeline/Source/MovieRenderPipelineEditor/Private/Widgets/SMoviePipelineQueueEditor.cpp
@@ -4,6 +4,7 @@
 #include "Widgets/SMoviePipelineQueueEditor.h"
 #include "Widgets/MoviePipelineWidgetConstants.h"
 #include "MovieRenderPipelineDataTypes.h"
+#include "MoviePipelineBasicConfig.h"
 #include "MoviePipelineQueue.h"
 #include "MoviePipelineOutputSetting.h"
 #include "MoviePipelineQueueSubsystem.h"
@@ -15,6 +16,7 @@
 #include "MoviePipelineEditorBlueprintLibrary.h"
 #include "Graph/MovieGraphConfig.h"
 #include "Graph/MovieGraphConfigFactory.h"
+#include "Graph/MovieRenderGraphEditorSettings.h"
 
 // Slate Includes
 #include "Widgets/Input/SComboButton.h"
@@ -267,6 +269,11 @@ public:
 		UMoviePipelineExecutorJob* Job = WeakJob.Get();
 		if (Job)
 		{
+			if (Job->IsUsingBasicConfiguration())
+			{
+				return LOCTEXT("BasicConfigLabel", "Basic");
+			}
+
 			if (Job->IsUsingGraphConfiguration())
 			{
 				if (Job->GetGraphPreset())
@@ -386,24 +393,22 @@ public:
 
 	}
 
-	void OnReplaceWithRenderGraph()
-	{	
+	void OnSwitchToGraph()
+	{
+		FScopedTransaction Transaction(LOCTEXT("SetGraphConfig_Transaction", "Set Graph Configuration"));
+
 		const bool bIsShot = false;
 		TSet<TSharedPtr<IMoviePipelineQueueTreeItem>> AffectedItems = GetMultiSelectAffectedItems(WeakQueueEditor, bIsShot);
-		
+
 		for (TSharedPtr<IMoviePipelineQueueTreeItem> Item : AffectedItems)
 		{
 			TSharedPtr<FMoviePipelineQueueJobTreeItem> JobTreeItem = Item->AsJob();
 			if (JobTreeItem.IsValid())
 			{
 				UMoviePipelineExecutorJob* Job = JobTreeItem->WeakJob.Get();
-				if (Job)
+				if (Job && !Job->SetConfigMode(EMoviePipelineConfigMode::Graph))
 				{
-					// Only replace jobs that don't already use the render-graph configuration
-					if(!Job->IsUsingGraphConfiguration())
-					{
-						SMoviePipelineQueueEditor::AssignDefaultGraphPresetToJob(Job);
-					}
+					SMoviePipelineQueueEditor::AssignDefaultGraphPresetToJob(Job);
 				}
 			}
 		}
@@ -436,13 +441,13 @@ public:
 		}
 	}
 
-	void OnClearGraph()
+	void OnSwitchToPreset()
 	{
-		FScopedTransaction Transaction(LOCTEXT("ClearJobGraph_Transaction", "Replace Graph with Config"));
+		FScopedTransaction Transaction(LOCTEXT("SetPresetConfig_Transaction", "Set Preset Configuration"));
 
 		const bool bIsShot = false;
 		TSet<TSharedPtr<IMoviePipelineQueueTreeItem>> AffectedItems = GetMultiSelectAffectedItems(WeakQueueEditor, bIsShot);
-		
+
 		for (TSharedPtr<IMoviePipelineQueueTreeItem> Item : AffectedItems)
 		{
 			TSharedPtr<FMoviePipelineQueueJobTreeItem> JobTreeItem = Item->AsJob();
@@ -451,7 +456,7 @@ public:
 				UMoviePipelineExecutorJob* Job = JobTreeItem->WeakJob.Get();
 				if (Job)
 				{
-					Job->SetGraphPreset(nullptr);
+					Job->SetConfigMode(EMoviePipelineConfigMode::Preset);
 				}
 			}
 		}
@@ -464,17 +469,38 @@ public:
 		UMoviePipelineExecutorJob* Job = WeakJob.Get();
 		if (Job)
 		{
+			// Basic-mode configs are always saved directly on the job, so there's no external asset to be dirty.
+			if (Job->IsUsingBasicConfiguration())
+			{
+				return EVisibility::Collapsed;
+			}
+
 			if (const UMovieGraphConfig* GraphPreset = Job->GetGraphPreset())
 			{
 				return GraphPreset->GetPackage()->IsDirty() ? EVisibility::Visible : EVisibility::Collapsed;
 			}
-			
+
 			return (Job->GetPresetOrigin() == nullptr) ? EVisibility::Visible : EVisibility::Collapsed;
 		}
-		
+
 		return EVisibility::Collapsed;
 	}
 
+	int32 GetConfigLabelWidgetIndex() const
+	{
+		UMoviePipelineExecutorJob* Job = WeakJob.Get();
+		return (Job && Job->IsUsingBasicConfiguration()) ? 0 : 1;
+	}
+
+	FText GetPrimaryConfigLabelWithModifiedIndicator() const
+	{
+		if (GetPrimaryConfigModifiedVisibility() == EVisibility::Visible)
+		{
+			return FText::Format(INVTEXT("{0} *"), GetPrimaryConfigLabel());
+		}
+		return GetPrimaryConfigLabel();
+	}
+
 	void OnEditPrimaryConfigForJob()
 	{
 		OnEditConfigCallback.ExecuteIfBound(WeakJob, nullptr);
@@ -553,9 +579,93 @@ public:
 		return Job ? FText::FromString(Job->GetStatusMessage()) : FText();
 	}
 
+	TSharedRef<SWidget> OnGenerateSequencePickerMenu()
+	{
+		FMenuBuilder MenuBuilder(true, nullptr);
+		IContentBrowserSingleton& ContentBrowser = FModuleManager::LoadModuleChecked<FContentBrowserModule>("ContentBrowser").Get();
+
+		FAssetPickerConfig AssetPickerConfig;
+		{
+			AssetPickerConfig.SelectionMode = ESelectionMode::Single;
+			AssetPickerConfig.InitialAssetViewType = EAssetViewType::List;
+			AssetPickerConfig.bFocusSearchBoxWhenOpened = true;
+			AssetPickerConfig.bAllowNullSelection = false;
+			AssetPickerConfig.bShowBottomToolbar = true;
+			AssetPickerConfig.bAutohideSearchBar = false;
+			AssetPickerConfig.bAllowDragging = false;
+			AssetPickerConfig.bCanShowClasses = false;
+			AssetPickerConfig.bShowPathInColumnView = true;
+			AssetPickerConfig.bShowTypeInColumnView = false;
+			AssetPickerConfig.bSortByPathInColumnView = false;
+			AssetPickerConfig.ThumbnailScale = 0.1f;
+			AssetPickerConfig.SaveSettingsName = TEXT("MoviePipelineSequenceAsset");
+
+			AssetPickerConfig.AssetShowWarningText = LOCTEXT("NoSequencesInPicker_Warning", "No Level Sequences Found");
+			AssetPickerConfig.Filter.ClassPaths.Add(ULevelSequence::StaticClass()->GetClassPathName());
+			AssetPickerConfig.Filter.bRecursiveClasses = true;
+			AssetPickerConfig.OnAssetSelected = FOnAssetSelected::CreateRaw(this, &FMoviePipelineQueueJobTreeItem::OnPickSequenceFromAsset);
+		}
+
+		MenuBuilder.BeginSection(NAME_None, LOCTEXT("PickSequence_MenuSection", "Level Sequence"));
+		{
+			TSharedRef<SWidget> PresetPicker = SNew(SBox)
+				.WidthOverride(300.f)
+				.HeightOverride(300.f)
+				[
+					ContentBrowser.CreateAssetPicker(AssetPickerConfig)
+				];
+
+			MenuBuilder.AddWidget(PresetPicker, FText(), true, false);
+		}
+		MenuBuilder.EndSection();
+
+		return MenuBuilder.MakeWidget();
+	}
+
+	void OnPickSequenceFromAsset(const FAssetData& AssetData)
+	{
+		FSlateApplication::Get().DismissAllMenus();
+
+		UMoviePipelineExecutorJob* CurrentJob = WeakJob.Get();
+		if (!CurrentJob)
+		{
+			return;
+		}
+
+		ULevelSequence* LevelSequence = Cast<ULevelSequence>(AssetData.GetAsset());
+		if (!LevelSequence)
+		{
+			return;
+		}
+
+		FScopedTransaction Transaction(LOCTEXT("PickJobSequence_Transaction", "Set Job Sequence"));
+
+		const bool bIsShot = false;
+		TSet<TSharedPtr<IMoviePipelineQueueTreeItem>> AffectedItems = GetMultiSelectAffectedItems(WeakQueueEditor, bIsShot);
+
+		for (TSharedPtr<IMoviePipelineQueueTreeItem> Item : AffectedItems)
+		{
+			TSharedPtr<FMoviePipelineQueueJobTreeItem> JobTreeItem = Item->AsJob();
+			if (JobTreeItem.IsValid())
+			{
+				UMoviePipelineExecutorJob* Job = JobTreeItem->WeakJob.Get();
+				if (Job)
+				{
+					Job->Modify();
+					Job->SetSequence(FSoftObjectPath(LevelSequence));
+					Job->JobName = LevelSequence->GetName();
+				}
+			}
+		}
+	}
+
 	TSharedRef<SWidget> OnGenerateConfigPresetPickerMenu()
 	{
 		UMoviePipelineExecutorJob* Job = WeakJob.Get();
+		if (!Job)
+		{
+			return SNullWidget::NullWidget;
+		}
 
 		UClass* ConfigType = Job->IsUsingGraphConfiguration() ? UMovieGraphConfig::StaticClass() : UMoviePipelinePrimaryConfig::StaticClass();
 		constexpr bool bIsShot = false;
@@ -566,16 +676,21 @@ public:
 			bIsShot,
 			FOnAssetSelected::CreateRaw(this, &FMoviePipelineQueueJobTreeItem::OnPickPresetFromAsset),
 			FExecuteAction::CreateRaw(this, &FMoviePipelineQueueJobTreeItem::OnClearNonGraphPreset),
-			FExecuteAction::CreateRaw(this, &FMoviePipelineQueueJobTreeItem::OnReplaceWithRenderGraph),
+			FExecuteAction::CreateRaw(this, &FMoviePipelineQueueJobTreeItem::OnSwitchToGraph),
 			FExecuteAction::CreateRaw(this, &FMoviePipelineQueueJobTreeItem::OnCreateNewGraphAndAssign),
-			FExecuteAction::CreateRaw(this, &FMoviePipelineQueueJobTreeItem::OnClearGraph)
+			FExecuteAction::CreateRaw(this, &FMoviePipelineQueueJobTreeItem::OnSwitchToPreset)
 			);
 	}
 
 	static TSharedRef<SWidget> OnGenerateConfigPresetPickerMenuFromClass(UClass* InClass,
 		TWeakObjectPtr<UMoviePipelineExecutorJob> TargetJob, const bool bIsShot, FOnAssetSelected InOnAssetSelected, FExecuteAction InClearNonGraphConfig,
-		FExecuteAction InNewRenderGraph, FExecuteAction InCreateNewGraphAndAssign, FExecuteAction InClearGraph)
+		FExecuteAction InSwitchToGraph, FExecuteAction InCreateNewGraphAndAssign, FExecuteAction InSwitchToPreset)
 	{
+		if (!TargetJob.IsValid())
+		{
+			return SNullWidget::NullWidget;
+		}
+
 		FMenuBuilder MenuBuilder(true, nullptr);
 		IContentBrowserSingleton& ContentBrowser = FModuleManager::LoadModuleChecked<FContentBrowserModule>("ContentBrowser").Get();
 
@@ -584,7 +699,7 @@ public:
 			const FText NoAssetsFoundWarning = InClass->IsChildOf(UMoviePipelineConfigBase::StaticClass())
 				? LOCTEXT("NoConfigs_Warning", "No Configurations Found")
 				: LOCTEXT("NoGraphs_Warning", "No Render Graphs Found");
-			
+
 			AssetPickerConfig.SelectionMode = ESelectionMode::Single;
 			AssetPickerConfig.InitialAssetViewType = EAssetViewType::List;
 			AssetPickerConfig.bFocusSearchBoxWhenOpened = true;
@@ -610,8 +725,8 @@ public:
 			if (!TargetJob->IsUsingGraphConfiguration())
 			{
 				MenuBuilder.AddMenuEntry(
-					LOCTEXT("ClearConfig_Label", "Clear Config"),
-					LOCTEXT("ClearConfig_Tooltip", "Resets the changes to the config and goes back to the defaults."),
+					LOCTEXT("ResetConfig_Label", "Reset Configuration"),
+					LOCTEXT("ResetConfig_Tooltip", "Resets the changes to the config and goes back to the defaults."),
 					FSlateIcon(),
 					FUIAction(InClearNonGraphConfig),
 					NAME_None,
@@ -620,14 +735,33 @@ public:
 			}
 			else
 			{
-				// Shots can clear their graph, meaning they will inherit from the parent graph. Primary jobs can revert to using the legacy system.
 				if (bIsShot)
 				{
+					// Shots can clear their graph, meaning they will inherit from the parent graph.
 					MenuBuilder.AddMenuEntry(
 						LOCTEXT("ClearGraph_Label", "Clear Graph"),
 						LOCTEXT("ClearGraph_Tooltip", "Remove the graph assigned to this shot. After removal, the shot will inherit the graph from the primary job."),
 						FSlateIcon(),
-						FUIAction(InClearGraph),
+						FUIAction(InSwitchToPreset),
+						NAME_None,
+						EUserInterfaceActionType::Button
+					);
+				}
+				else
+				{
+					// Primary jobs in graph mode: "Reset Configuration" reassigns the default graph from project settings.
+					MenuBuilder.AddMenuEntry(
+						LOCTEXT("ResetGraphConfig_Label", "Reset Configuration"),
+						LOCTEXT("ResetGraphConfig_Tooltip", "Resets the graph to the default render graph specified in project settings."),
+						FSlateIcon(),
+						FUIAction(FExecuteAction::CreateLambda([TargetJob]()
+						{
+							if (TargetJob.IsValid())
+							{
+								SMoviePipelineQueueEditor::AssignDefaultGraphPresetToJob(TargetJob.Get());
+								FSlateApplication::Get().DismissAllMenus();
+							}
+						})),
 						NAME_None,
 						EUserInterfaceActionType::Button
 					);
@@ -640,34 +774,66 @@ public:
 		{
 			MenuBuilder.BeginSection(NAME_None, LOCTEXT("ConfigurationType_MenuSection", "Configuration Type"));
 			{
+				// Basic radio button — only available in the new queue panel
+				static IConsoleVariable* ShowNewQueuePanelCVar = IConsoleManager::Get().FindConsoleVariable(TEXT("MovieRenderPipeline.ShowNewQueuePanel"));
+				if (ShowNewQueuePanelCVar && ShowNewQueuePanelCVar->GetInt() != 0)
+				{
+					FUIAction SwitchToBasicAction;
+					SwitchToBasicAction.ExecuteAction = FExecuteAction::CreateLambda([TargetJob]()
+					{
+						if (TargetJob.IsValid())
+						{
+							FScopedTransaction Transaction(LOCTEXT("SetBasicConfig", "Set Basic Configuration"));
+							TargetJob->SetConfigMode(EMoviePipelineConfigMode::Basic);
+							FSlateApplication::Get().DismissAllMenus();
+						}
+					});
+					SwitchToBasicAction.GetActionCheckState = FGetActionCheckState::CreateLambda([TargetJob]() -> ECheckBoxState
+					{
+						return (TargetJob.IsValid() && TargetJob->IsUsingBasicConfiguration()) ? ECheckBoxState::Checked : ECheckBoxState::Unchecked;
+					});
+
+					MenuBuilder.AddMenuEntry(
+						LOCTEXT("ConfigurationType_BasicLabel", "Basic"),
+						LOCTEXT("ConfigurationType_BasicTooltip", "Use a simplified set of rendering settings."),
+						FSlateIcon(),
+						SwitchToBasicAction,
+						NAME_None,
+						EUserInterfaceActionType::RadioButton
+					);
+				}
+
+				// Preset radio button
 				FUIAction SwitchToPresetAction;
-				SwitchToPresetAction.ExecuteAction = InClearGraph;
+				SwitchToPresetAction.ExecuteAction = InSwitchToPreset;
 				SwitchToPresetAction.GetActionCheckState = FGetActionCheckState::CreateLambda([TargetJob]() -> ECheckBoxState
 				{
-					return (TargetJob.IsValid() && TargetJob->IsUsingGraphConfiguration()) ? ECheckBoxState::Unchecked : ECheckBoxState::Checked;
+					return (TargetJob.IsValid() && !TargetJob->IsUsingGraphConfiguration() && !TargetJob->IsUsingBasicConfiguration())
+						? ECheckBoxState::Checked : ECheckBoxState::Unchecked;
 				});
 
+				// Graph radio button
 				FUIAction SwitchToGraphAction;
-				SwitchToGraphAction.ExecuteAction = InNewRenderGraph;
+				SwitchToGraphAction.ExecuteAction = InSwitchToGraph;
 				SwitchToGraphAction.GetActionCheckState = FGetActionCheckState::CreateLambda([TargetJob]() -> ECheckBoxState
 				{
 					return (TargetJob.IsValid() && TargetJob->IsUsingGraphConfiguration()) ? ECheckBoxState::Checked : ECheckBoxState::Unchecked;
 				});
-			
+
 				MenuBuilder.AddMenuEntry(
-					LOCTEXT("ConfigurationType_PresetLabel", "Preset"),
-					LOCTEXT("ConfigurationType_PresetTooltip", "Updates this job or shot to use a preset-based configuration."),
+					LOCTEXT("ConfigurationType_GraphLabel", "Movie Render Graph"),
+					LOCTEXT("ConfigurationType_GraphTooltip", "Updates this job or shot to use a graph-based configuration."),
 					FSlateIcon(),
-					SwitchToPresetAction,
+					SwitchToGraphAction,
 					NAME_None,
 					EUserInterfaceActionType::RadioButton
 				);
 
 				MenuBuilder.AddMenuEntry(
-					LOCTEXT("ConfigurationType_GraphLabel", "Movie Render Graph (Beta)"),
-					LOCTEXT("ConfigurationType_GraphTooltip", "Updates this job or shot to use a graph-based configuration."),
+					LOCTEXT("ConfigurationType_PresetLabel", "Legacy Preset"),
+					LOCTEXT("ConfigurationType_PresetTooltip", "Updates this job or shot to use a preset-based configuration."),
 					FSlateIcon(),
-					SwitchToGraphAction,
+					SwitchToPresetAction,
 					NAME_None,
 					EUserInterfaceActionType::RadioButton
 				);
@@ -675,9 +841,8 @@ public:
 			MenuBuilder.EndSection();
 		}
 
-		// If this menu is being created for a shot, only allow creating a new graph via the menu if the parent job is using a graph. Creating a
-		// graph for a shot with a parent using a preset-based configuration is not valid.
-		if (!bIsShot || (bIsShot && TargetJob->IsUsingGraphConfiguration()))
+		// "Create New Asset" is only relevant in Graph mode. Shots additionally require the parent job to be in Graph mode.
+		if (TargetJob->IsUsingGraphConfiguration())
 		{
 			MenuBuilder.BeginSection(NAME_None, LOCTEXT("CreateNewAsset_MenuSection", "Create New Asset"));
 			{
@@ -693,18 +858,22 @@ public:
 			MenuBuilder.EndSection();
 		}
 
-		MenuBuilder.BeginSection(NAME_None, LOCTEXT("ImportConfig_MenuSection", "Import Configuration"));
+		// Asset picker is not applicable in Basic mode (Basic config is edited directly in the details panel).
+		if (!TargetJob->IsUsingBasicConfiguration())
 		{
-			TSharedRef<SWidget> PresetPicker = SNew(SBox)
-				.WidthOverride(300.f)
-				.HeightOverride(300.f)
-				[
-					ContentBrowser.CreateAssetPicker(AssetPickerConfig)
-				];
-
-			MenuBuilder.AddWidget(PresetPicker, FText(), true, false);
+			MenuBuilder.BeginSection(NAME_None, LOCTEXT("ImportConfig_MenuSection", "Import Configuration"));
+			{
+				TSharedRef<SWidget> PresetPicker = SNew(SBox)
+					.WidthOverride(300.f)
+					.HeightOverride(300.f)
+					[
+						ContentBrowser.CreateAssetPicker(AssetPickerConfig)
+					];
+
+				MenuBuilder.AddWidget(PresetPicker, FText(), true, false);
+			}
+			MenuBuilder.EndSection();
 		}
-		MenuBuilder.EndSection();
 
 		return MenuBuilder.MakeWidget();
 	}
@@ -763,10 +932,36 @@ TSharedRef<SWidget> SQueueJobListRow::GenerateWidgetForColumn(const FName& Colum
 
 				+SHorizontalBox::Slot()
 				.FillWidth(1.f)
+				.VAlign(VAlign_Center)
 				[
 					SNew(STextBlock)
 					.Text(Item.Get(), &FMoviePipelineQueueJobTreeItem::GetJobName)
 				]
+
+				// Sequence picker dropdown
+				+SHorizontalBox::Slot()
+				.AutoWidth()
+				.VAlign(VAlign_Center)
+				.HAlign(HAlign_Right)
+				.Padding(4, 0, 4, 0)
+				[
+					SNew(SComboButton)
+					.ContentPadding(1)
+					.IsEnabled(Item.Get(), &FMoviePipelineQueueJobTreeItem::IsConfigEditingEnabled)
+					.OnGetMenuContent(Item.Get(), &FMoviePipelineQueueJobTreeItem::OnGenerateSequencePickerMenu)
+					.HasDownArrow(false)
+					.ButtonContent()
+					[
+						SNew(SBox)
+						.Padding(FMargin(2, 0))
+						[
+							SNew(STextBlock)
+							.TextStyle(FAppStyle::Get(), "NormalText.Important")
+							.Font(FAppStyle::Get().GetFontStyle("FontAwesome.10"))
+							.Text(FEditorFontGlyphs::Caret_Down)
+						]
+					]
+				]
 			];
 	}
 	else if (ColumnName == NAME_Settings)
@@ -776,27 +971,28 @@ TSharedRef<SWidget> SQueueJobListRow::GenerateWidgetForColumn(const FName& Colum
 
 		// Preset Label
 		+ SHorizontalBox::Slot()
-		.AutoWidth()
+		.FillWidth(1.f)
+		.HAlign(HAlign_Left)
 		.VAlign(VAlign_Center)
 		.Padding(2, 0)
 		[
-			SNew(SHyperlink)
-			.Text(TAttribute<FText>::Create(TAttribute<FText>::FGetter::CreateSP(Item.Get(), &FMoviePipelineQueueJobTreeItem::GetPrimaryConfigLabel)))
-			.OnNavigate(Item.Get(), &FMoviePipelineQueueJobTreeItem::OnEditPrimaryConfigForJob)
-		]
+			SNew(SWidgetSwitcher)
+			.WidgetIndex(Item.Get(), &FMoviePipelineQueueJobTreeItem::GetConfigLabelWidgetIndex)
 
-		+ SHorizontalBox::Slot()
-		.AutoWidth()
-		[
-			SNew(STextBlock)
-			.Text(LOCTEXT("ModifiedConfigIndicator", "*"))
-			.Visibility(Item.Get(), &FMoviePipelineQueueJobTreeItem::GetPrimaryConfigModifiedVisibility)
-		]
+			// Index 0: Basic config — plain text, no hyperlink
+			+ SWidgetSwitcher::Slot()
+			[
+				SNew(STextBlock)
+				.Text(TAttribute<FText>::Create(TAttribute<FText>::FGetter::CreateSP(Item.Get(), &FMoviePipelineQueueJobTreeItem::GetPrimaryConfigLabel)))
+			]
 
-		+ SHorizontalBox::Slot()
-		.FillWidth(1.f)
-		[
-			SNullWidget::NullWidget
+			// Index 1: Preset/graph config — hyperlink to open preset picker
+			+ SWidgetSwitcher::Slot()
+			[
+				SNew(SHyperlink)
+				.Text(TAttribute<FText>::Create(TAttribute<FText>::FGetter::CreateSP(Item.Get(), &FMoviePipelineQueueJobTreeItem::GetPrimaryConfigLabelWithModifiedIndicator)))
+				.OnNavigate(Item.Get(), &FMoviePipelineQueueJobTreeItem::OnEditPrimaryConfigForJob)
+			]
 		]
 
 		// Dropdown Arrow
@@ -1121,6 +1317,17 @@ struct FMoviePipelineShotItem : IMoviePipelineQueueTreeItem
 		UMoviePipelineExecutorShot* Shot = WeakShot.Get();
 		if (Shot)
 		{
+			if (Shot->IsUsingBasicConfiguration())
+			{
+				const UMoviePipelineBasicConfig* ShotConfig = Shot->GetBasicShotConfig();
+				if (ShotConfig && ShotConfig->HasAnyOverrides())
+				{
+					return LOCTEXT("BasicShotOverrideLabel", "Basic (Override)");
+				}
+
+				return LOCTEXT("BasicShotInheritLabel", "Basic");
+			}
+
 			if (Shot->IsUsingGraphConfiguration())
 			{
 				if (const UMovieGraphConfig* GraphConfig = Shot->GetGraphPreset())
@@ -1314,11 +1521,17 @@ struct FMoviePipelineShotItem : IMoviePipelineQueueTreeItem
 	{
 		if (const UMoviePipelineExecutorShot* Shot = WeakShot.Get())
 		{
+			// Basic-mode configs are always saved directly on the job, so there's no external asset to be dirty.
+			if (Shot->IsUsingBasicConfiguration())
+			{
+				return EVisibility::Collapsed;
+			}
+
 			if (const UMovieGraphConfig* GraphPreset = Shot->GetGraphPreset())
 			{
 				return GraphPreset->GetPackage()->IsDirty() ? EVisibility::Visible : EVisibility::Collapsed;
 			}
-		
+
 			if (Shot->GetShotOverrideConfiguration() && (Shot->GetShotOverridePresetOrigin() == nullptr))
 			{
 				return EVisibility::Visible;
@@ -1328,6 +1541,21 @@ struct FMoviePipelineShotItem : IMoviePipelineQueueTreeItem
 		return EVisibility::Collapsed;
 	}
 
+	FText GetShotConfigLabelWithModifiedIndicator() const
+	{
+		if (GetShotConfigModifiedVisibility() == EVisibility::Visible)
+		{
+			return FText::Format(INVTEXT("{0} *"), GetShotConfigLabel());
+		}
+		return GetShotConfigLabel();
+	}
+
+	int32 GetShotConfigLabelWidgetIndex() const
+	{
+		UMoviePipelineExecutorShot* Shot = WeakShot.Get();
+		return (Shot && Shot->IsUsingBasicConfiguration()) ? 0 : 1;
+	}
+
 	void OnEditConfigForShot()
 	{
 		OnEditConfigCallback.ExecuteIfBound(WeakJob, WeakShot);
@@ -1338,15 +1566,15 @@ struct FMoviePipelineShotItem : IMoviePipelineQueueTreeItem
 		const UMoviePipelineExecutorShot* Shot = WeakShot.Get();
 		UClass* ConfigType = Shot && Shot->IsUsingGraphConfiguration() ? UMovieGraphConfig::StaticClass() : UMoviePipelineShotConfig::StaticClass();
 		constexpr bool bIsShot = true;
-		const FExecuteAction ReplaceWithGraph = nullptr;
-		
+		const FExecuteAction SwitchToGraph = nullptr;
+
 		return FMoviePipelineQueueJobTreeItem::OnGenerateConfigPresetPickerMenuFromClass(
 			ConfigType,
 			WeakJob,
 			bIsShot,
 			FOnAssetSelected::CreateRaw(this, &FMoviePipelineShotItem::OnPickShotPresetFromAsset),
 			FExecuteAction::CreateRaw(this, &FMoviePipelineShotItem::OnClearNonGraphPreset),
-			ReplaceWithGraph,
+			SwitchToGraph,
 			FExecuteAction::CreateRaw(this, &FMoviePipelineShotItem::OnCreateNewGraphAndAssign),
 			FExecuteAction::CreateRaw(this, &FMoviePipelineShotItem::OnClearGraph)
 			);
@@ -1418,27 +1646,28 @@ TSharedRef<SWidget> SQueueShotListRow::GenerateWidgetForColumn(const FName& Colu
 
 		// Preset Label
 		+ SHorizontalBox::Slot()
-		.AutoWidth()
+		.FillWidth(1.f)
+		.HAlign(HAlign_Left)
 		.VAlign(VAlign_Center)
 		.Padding(2, 0)
 		[
-			SNew(SHyperlink)
-			.Text(TAttribute<FText>::Create(TAttribute<FText>::FGetter::CreateSP(Item.Get(), &FMoviePipelineShotItem::GetShotConfigLabel)))
-			.OnNavigate(Item.Get(), &FMoviePipelineShotItem::OnEditConfigForShot)
-		]
+			SNew(SWidgetSwitcher)
+			.WidgetIndex(Item.Get(), &FMoviePipelineShotItem::GetShotConfigLabelWidgetIndex)
 
-		+ SHorizontalBox::Slot()
-		.AutoWidth()
-		[
-			SNew(STextBlock)
-			.Text(LOCTEXT("ModifiedShotConfigIndicator", "*"))
-			.Visibility(Item.Get(), &FMoviePipelineShotItem::GetShotConfigModifiedVisibility)
-		]
+			// Index 0: Basic config — plain text, no hyperlink
+			+ SWidgetSwitcher::Slot()
+			[
+				SNew(STextBlock)
+				.Text(TAttribute<FText>::Create(TAttribute<FText>::FGetter::CreateSP(Item.Get(), &FMoviePipelineShotItem::GetShotConfigLabel)))
+			]
 
-		+ SHorizontalBox::Slot()
-		.FillWidth(1.f)
-		[
-			SNullWidget::NullWidget
+			// Index 1: Preset/graph config — hyperlink to open preset picker
+			+ SWidgetSwitcher::Slot()
+			[
+				SNew(SHyperlink)
+				.Text(TAttribute<FText>::Create(TAttribute<FText>::FGetter::CreateSP(Item.Get(), &FMoviePipelineShotItem::GetShotConfigLabelWithModifiedIndicator)))
+				.OnNavigate(Item.Get(), &FMoviePipelineShotItem::OnEditConfigForShot)
+			]
 		]
 
 		// Dropdown Arrow
@@ -1735,28 +1964,30 @@ void SMoviePipelineQueueEditor::OnCreateJobFromAsset(const FAssetData& InAsset)
 	}
 
 	const UMovieRenderPipelineProjectSettings* ProjectSettings = GetDefault<UMovieRenderPipelineProjectSettings>();
-	const UClass* DefaultPipeline = Cast<UClass>(ProjectSettings->DefaultPipeline.TryLoad());
+	const UMovieRenderGraphEditorSettings* EditorSettings = GetDefault<UMovieRenderGraphEditorSettings>();
 	for (UMoviePipelineExecutorJob* NewJob : NewJobs)
 	{
 		PendingJobsToSelect.Add(NewJob);
 		
+		// Set up configuration based on the user's default configuration type preference
+		switch (EditorSettings->DefaultConfigType)
 		{
-			// The job configuration is already set up with an empty configuration, but we'll try and use their last used preset 
-			// (or a engine supplied default) for better user experience. 
+		case EMoviePipelineDefaultConfigType::Basic:
+			NewJob->SetupBasicConfiguration();
+			break;
+		case EMoviePipelineDefaultConfigType::MovieRenderGraph:
+			SMoviePipelineQueueEditor::AssignDefaultGraphPresetToJob(NewJob);
+			break;
+		case EMoviePipelineDefaultConfigType::LegacyPreset:
 			if (ProjectSettings->LastPresetOrigin.IsValid())
 			{
 				NewJob->SetPresetOrigin(ProjectSettings->LastPresetOrigin.Get());
 			}
-		}
-
-		// Ensure the job has the settings specified by the project settings added. If they're already added
-		// we don't modify the object so that we don't make it confused about whether or not you've modified the preset.
-		UMoviePipelineEditorBlueprintLibrary::EnsureJobHasDefaultSettings(NewJob);
-
-		// If the default class is a movie graph, assign the default graph
-		if (DefaultPipeline && DefaultPipeline == UMovieGraphPipeline::StaticClass())
-		{
-			AssignDefaultGraphPresetToJob(NewJob);
+			UMoviePipelineEditorBlueprintLibrary::EnsureJobHasDefaultSettings(NewJob);
+			break;
+		default:
+			checkNoEntry();
+			break;
 		}
 	}
 }
