[MoviePipeline] Added InOutOverscanFraction argument to FMovieGraphRenderCameraSource::GetCameraOverscannedResolution().

After FMovieGraphDeferredPass::GetResolutionAndCameraInfo() returns,
FMovieGraphDeferredPass::Render() applies the overscan fraction to the camera ViewInfo:

CameraInfo.ViewInfo.ClearOverscan();
CameraInfo.ViewInfo.ApplyOverscan(ResolutionAndCameraInfo.OverscanFraction);

Render camera source implementations that handle overscan internally must prevent this
from double-applying the effect. By setting InOutOverscanFraction to 0 inside
GetCameraOverscannedResolution, an implementation can skip the ViewInfo overscan step
while still routing the correct value through its own pipeline.

#Synced-CL 0

diff --git a/Engine/Plugins/MovieScene/MovieRenderPipeline/Source/MovieRenderPipelineCore/Public/Graph/MovieGraphRenderCameraSource.h b/Engine/Plugins/MovieScene/MovieRenderPipeline/Source/MovieRenderPipelineCore/Public/Graph/MovieGraphRenderCameraSource.h
--- a/Engine/Plugins/MovieScene/MovieRenderPipeline/Source/MovieRenderPipelineCore/Public/Graph/MovieGraphRenderCameraSource.h
+++ b/Engine/Plugins/MovieScene/MovieRenderPipeline/Source/MovieRenderPipelineCore/Public/Graph/MovieGraphRenderCameraSource.h
@@ -90,15 +90,17 @@ struct FMovieGraphRenderCameraSource
 	/**
 	 * Returns the overscanned resolution and inner crop rectangle for the specified camera.
 	 *
-	 * @param InCameraIndex            (in)  Index in range [0, NumCameras).
-	 * @param OutOverscannedResolution (out) Receives the overscanned resolution.
-	 * @param OutOverscanCropRectangle (out) Receives the inner crop rectangle within the overscanned resolution.
-	 * @return                         True if resolved; false to fall back to the default resolution calculation.
+	 * @param InCameraIndex              (in)  Index in range [0, NumCameras).
+	 * @param InOutOverscannedResolution (in,out) Receives the overscanned resolution.
+	 * @param InOutOverscanCropRectangle (in,out) Receives the inner crop rectangle within the overscanned resolution.
+	 * @param InOutOverscanFraction      (in,out) The final per-side overscan fraction [0..1] for this camera.
+	 * @return                           True if resolved; false to fall back to the default resolution calculation.
 	 */
 	UE_API virtual bool GetCameraOverscannedResolution(
 		const int32 InCameraIndex,
-		FIntPoint& OutOverscannedResolution,
-		FIntRect& OutOverscanCropRectangle) const { return false; }
+		FIntPoint&  InOutOverscannedResolution,
+		FIntRect&   InOutOverscanCropRectangle,
+		float&      InOutOverscanFraction) const { return false; }
 };
 
 #undef UE_API
\ No newline at end of file
diff --git a/Engine/Plugins/MovieScene/MovieRenderPipeline/Source/MovieRenderPipelineRenderPasses/Private/Graph/Renderers/MovieGraphDeferredPass.cpp b/Engine/Plugins/MovieScene/MovieRenderPipeline/Source/MovieRenderPipelineRenderPasses/Private/Graph/Renderers/MovieGraphDeferredPass.cpp
--- a/Engine/Plugins/MovieScene/MovieRenderPipeline/Source/MovieRenderPipelineRenderPasses/Private/Graph/Renderers/MovieGraphDeferredPass.cpp
+++ b/Engine/Plugins/MovieScene/MovieRenderPipeline/Source/MovieRenderPipelineRenderPasses/Private/Graph/Renderers/MovieGraphDeferredPass.cpp
@@ -139,10 +139,11 @@ FMovieGraphDeferredPass::FResolutionAndCameraInfo FMovieGraphDeferredPass::GetRe
 	// Specifies a crop within the AccumulatorResolution for us to take center-out crops later where needed.
 	Info.AccumulatorResolutionCropRect = UMovieGraphBlueprintLibrary::GetOverscanCropRectangle(InEvaluatedConfig, Info.OverscanFraction, Info.CameraAspectRatio);
 
-	// If a custom camera source is active, let it override the accumulator resolution and crop rect
-	// computed above with per-camera values.
+	// If a custom camera source is active, let it override the accumulator resolution, crop rect,
+	// and overscan fraction computed above with per-camera values.
 	const TSharedPtr<const FMovieGraphRenderCameraSource> CameraSource = GetRenderCameraSource();
-	if (CameraSource && CameraSource->GetCameraOverscannedResolution(GetRenderCameraIndex(), Info.AccumulatorResolution, Info.AccumulatorResolutionCropRect))
+	if (CameraSource && CameraSource->GetCameraOverscannedResolution(GetRenderCameraIndex(),
+		Info.AccumulatorResolution, Info.AccumulatorResolutionCropRect, Info.OverscanFraction))
 	{
 		// A custom camera source provides multiple cameras, so bRenderAllCameras must be
 		// set to override the default single-camera path and render all of them.
